Copying //Tasks/UE4/Dev-Online-Plugins-v3 to //UE4/Dev-Main (Source: //Tasks/UE4/Dev-Online-Plugins-v3 @ 3047294)

diff --git a/Engine/Plugins/Online/OnlineSubsystemSteam/Source/Public/OnlineSubsystemSteamModule.h b/Engine/Plugins/Online/OnlineSubsystemSteam/Source/Public/OnlineSubsystemSteamModule.h
new file mode 100644
--- /dev/null
+++ b/Engine/Plugins/Online/OnlineSubsystemSteam/Source/Public/OnlineSubsystemSteamModule.h
@@ -0,0 +1,71 @@
+// Copyright 1998-2016 Epic Games, Inc. All Rights Reserved.
+
+#pragma once
+
+#include "Core.h"
+#include "ModuleInterface.h"
+
+#define LOADING_STEAM_CLIENT_LIBRARY_DYNAMICALLY				(PLATFORM_WINDOWS || PLATFORM_MAC || (PLATFORM_LINUX && !IS_MONOLITHIC))
+#define LOADING_STEAM_SERVER_LIBRARY_DYNAMICALLY				((PLATFORM_WINDOWS && PLATFORM_32BITS) || (PLATFORM_LINUX && !IS_MONOLITHIC) || PLATFORM_MAC)
+#define LOADING_STEAM_LIBRARIES_DYNAMICALLY						(LOADING_STEAM_CLIENT_LIBRARY_DYNAMICALLY || LOADING_STEAM_SERVER_LIBRARY_DYNAMICALLY)
+
+/**
+ * Online subsystem module class  (STEAM Implementation)
+ * Code related to the loading of the STEAM module
+ */
+class FOnlineSubsystemSteamModule : public IModuleInterface
+{
+private:
+
+	/** Class responsible for creating instance(s) of the subsystem */
+	class FOnlineFactorySteam* SteamFactory;
+
+#if LOADING_STEAM_LIBRARIES_DYNAMICALLY
+	/** Handle to the STEAM API dll */
+	void* SteamDLLHandle;
+	/** Handle to the STEAM dedicated server support dlls */
+	void* SteamServerDLLHandle;
+#endif	//LOADING_STEAM_LIBRARIES_DYNAMICALLY
+
+	/**
+	 *	Load the required modules for Steam
+	 */
+	void LoadSteamModules();
+
+	/** 
+	 *	Unload the required modules for Steam
+	 */
+	void UnloadSteamModules();
+
+public:
+
+	FOnlineSubsystemSteamModule()
+        : SteamFactory(NULL)
+#if LOADING_STEAM_LIBRARIES_DYNAMICALLY
+		, SteamDLLHandle(NULL)
+		, SteamServerDLLHandle(NULL)
+#endif	//LOADING_STEAM_LIBRARIES_DYNAMICALLY
+	{}
+
+	virtual ~FOnlineSubsystemSteamModule() {}
+
+	// IModuleInterface
+
+	virtual void StartupModule() override;
+	virtual void ShutdownModule() override;
+
+	virtual bool SupportsDynamicReloading() override
+	{
+		return false;
+	}
+
+	virtual bool SupportsAutomaticShutdown() override
+	{
+		return false;
+	}
+
+	/**
+	 * Are the Steam support Dlls loaded
+	 */
+	bool AreSteamDllsLoaded() const;
+};
\ No newline at end of file
diff --git a/Engine/Plugins/Online/OnlineSubsystemUtils/OnlineSubsystemUtils.uplugin b/Engine/Plugins/Online/OnlineSubsystemUtils/OnlineSubsystemUtils.uplugin
new file mode 100644
--- /dev/null
+++ b/Engine/Plugins/Online/OnlineSubsystemUtils/OnlineSubsystemUtils.uplugin
@@ -0,0 +1,26 @@
+{
+	"FileVersion" : 3,
+	
+	"FriendlyName" : "Online Subsystem Utils",
+	"Version" : 1,
+	"VersionName" : "1.0",
+	"EngineVersion" : "4.2.0",
+	"Description" : "Shared code for interacting online subsystem implementations.",
+	"Category" : "Online Platform",
+	"CreatedBy" : "Epic Games, Inc.",
+	"CreatedByURL" : "http://epicgames.com",
+	"EnabledByDefault" : true,
+
+    "Modules": [
+        {
+            "Name": "OnlineSubsystemUtils",
+            "Type": "Runtime",
+            "LoadingPhase": "PostConfigInit"
+        },
+        {
+            "Name": "OnlineBlueprintSupport",
+            "Type": "Developer",
+            "LoadingPhase": "PostConfigInit"
+        }
+    ]
+}
\ No newline at end of file
diff --git a/Engine/Plugins/Online/OnlineSubsystemUtils/Source/OnlineBlueprintSupport/Classes/K2Node_InAppPurchase.h b/Engine/Plugins/Online/OnlineSubsystemUtils/Source/OnlineBlueprintSupport/Classes/K2Node_InAppPurchase.h
new file mode 100644
--- /dev/null
+++ b/Engine/Plugins/Online/OnlineSubsystemUtils/Source/OnlineBlueprintSupport/Classes/K2Node_InAppPurchase.h
@@ -0,0 +1,12 @@
+// Copyright 1998-2016 Epic Games, Inc. All Rights Reserved.
+
+#pragma once
+
+#include "Classes/K2Node_BaseAsyncTask.h"
+#include "K2Node_InAppPurchase.generated.h"
+
+UCLASS()
+class ONLINEBLUEPRINTSUPPORT_API UK2Node_InAppPurchase : public UK2Node_BaseAsyncTask
+{
+	GENERATED_UCLASS_BODY()
+};
diff --git a/Engine/Plugins/Online/OnlineSubsystemUtils/Source/OnlineBlueprintSupport/Classes/K2Node_InAppPurchaseQuery.h b/Engine/Plugins/Online/OnlineSubsystemUtils/Source/OnlineBlueprintSupport/Classes/K2Node_InAppPurchaseQuery.h
new file mode 100644
--- /dev/null
+++ b/Engine/Plugins/Online/OnlineSubsystemUtils/Source/OnlineBlueprintSupport/Classes/K2Node_InAppPurchaseQuery.h
@@ -0,0 +1,12 @@
+// Copyright 1998-2016 Epic Games, Inc. All Rights Reserved.
+
+#pragma once
+
+#include "K2Node_BaseAsyncTask.h"
+#include "K2Node_InAppPurchaseQuery.generated.h"
+
+UCLASS()
+class ONLINEBLUEPRINTSUPPORT_API UK2Node_InAppPurchaseQuery : public UK2Node_BaseAsyncTask
+{
+	GENERATED_UCLASS_BODY()
+};
diff --git a/Engine/Plugins/Online/OnlineSubsystemUtils/Source/OnlineBlueprintSupport/Classes/K2Node_InAppPurchaseRestore.h b/Engine/Plugins/Online/OnlineSubsystemUtils/Source/OnlineBlueprintSupport/Classes/K2Node_InAppPurchaseRestore.h
new file mode 100644
--- /dev/null
+++ b/Engine/Plugins/Online/OnlineSubsystemUtils/Source/OnlineBlueprintSupport/Classes/K2Node_InAppPurchaseRestore.h
@@ -0,0 +1,12 @@
+// Copyright 1998-2016 Epic Games, Inc. All Rights Reserved.
+
+#pragma once
+
+#include "K2Node_BaseAsyncTask.h"
+#include "K2Node_InAppPurchaseRestore.generated.h"
+
+UCLASS()
+class ONLINEBLUEPRINTSUPPORT_API UK2Node_InAppPurchaseRestore : public UK2Node_BaseAsyncTask
+{
+	GENERATED_UCLASS_BODY()
+};
diff --git a/Engine/Plugins/Online/OnlineSubsystemUtils/Source/OnlineBlueprintSupport/Classes/K2Node_LeaderboardFlush.h b/Engine/Plugins/Online/OnlineSubsystemUtils/Source/OnlineBlueprintSupport/Classes/K2Node_LeaderboardFlush.h
new file mode 100644
--- /dev/null
+++ b/Engine/Plugins/Online/OnlineSubsystemUtils/Source/OnlineBlueprintSupport/Classes/K2Node_LeaderboardFlush.h
@@ -0,0 +1,21 @@
+// Copyright 1998-2016 Epic Games, Inc. All Rights Reserved.
+
+#pragma once
+
+#include "K2Node_BaseAsyncTask.h"
+#include "K2Node_LeaderboardFlush.generated.h"
+
+UCLASS()
+class ONLINEBLUEPRINTSUPPORT_API UK2Node_LeaderboardFlush : public UK2Node_BaseAsyncTask
+{
+	GENERATED_UCLASS_BODY()
+
+	//~ Begin UEdGraphNode Interface
+	virtual FText GetTooltipText() const override;
+	virtual FText GetNodeTitle(ENodeTitleType::Type TitleType) const override;
+	//~ End UEdGraphNode Interface
+
+	//~ Begin UK2Node Interface
+	virtual FText GetMenuCategory() const override;
+	//~ End UK2Node Interface
+};
diff --git a/Engine/Plugins/Online/OnlineSubsystemUtils/Source/OnlineBlueprintSupport/Classes/K2Node_LeaderboardQuery.h b/Engine/Plugins/Online/OnlineSubsystemUtils/Source/OnlineBlueprintSupport/Classes/K2Node_LeaderboardQuery.h
new file mode 100644
--- /dev/null
+++ b/Engine/Plugins/Online/OnlineSubsystemUtils/Source/OnlineBlueprintSupport/Classes/K2Node_LeaderboardQuery.h
@@ -0,0 +1,12 @@
+// Copyright 1998-2016 Epic Games, Inc. All Rights Reserved.
+
+#pragma once
+
+#include "K2Node_BaseAsyncTask.h"
+#include "K2Node_LeaderboardQuery.generated.h"
+
+UCLASS()
+class ONLINEBLUEPRINTSUPPORT_API UK2Node_LeaderboardQuery : public UK2Node_BaseAsyncTask
+{
+	GENERATED_UCLASS_BODY()
+};
diff --git a/Engine/Plugins/Online/OnlineSubsystemUtils/Source/OnlineBlueprintSupport/OnlineBlueprintSupport.Build.cs b/Engine/Plugins/Online/OnlineSubsystemUtils/Source/OnlineBlueprintSupport/OnlineBlueprintSupport.Build.cs
new file mode 100644
--- /dev/null
+++ b/Engine/Plugins/Online/OnlineSubsystemUtils/Source/OnlineBlueprintSupport/OnlineBlueprintSupport.Build.cs
@@ -0,0 +1,22 @@
+// Copyright 1998-2016 Epic Games, Inc. All Rights Reserved.
+
+using UnrealBuildTool;
+
+public class OnlineBlueprintSupport : ModuleRules
+{
+	public OnlineBlueprintSupport(TargetInfo Target)
+	{
+        PrivateDependencyModuleNames.AddRange(
+			new string[] { 
+				"Core", 
+				"CoreUObject", 
+				"Engine",
+				"BlueprintGraph",
+				"OnlineSubsystem",
+				"OnlineSubsystemUtils",
+				"Sockets",
+				"Json"
+			}
+		);
+	}
+}
diff --git a/Engine/Plugins/Online/OnlineSubsystemUtils/Source/OnlineBlueprintSupport/Private/K2Node_InAppPurchase.cpp b/Engine/Plugins/Online/OnlineSubsystemUtils/Source/OnlineBlueprintSupport/Private/K2Node_InAppPurchase.cpp
new file mode 100644
--- /dev/null
+++ b/Engine/Plugins/Online/OnlineSubsystemUtils/Source/OnlineBlueprintSupport/Private/K2Node_InAppPurchase.cpp
@@ -0,0 +1,18 @@
+// Copyright 1998-2016 Epic Games, Inc. All Rights Reserved.
+
+#include "OnlineBlueprintSupportPrivatePCH.h"
+#include "InAppPurchaseCallbackProxy.h"
+#include "K2Node_InAppPurchase.h"
+
+#define LOCTEXT_NAMESPACE "K2Node"
+
+UK2Node_InAppPurchase::UK2Node_InAppPurchase(const FObjectInitializer& ObjectInitializer)
+	: Super(ObjectInitializer)
+{
+	ProxyFactoryFunctionName = GET_FUNCTION_NAME_CHECKED(UInAppPurchaseCallbackProxy, CreateProxyObjectForInAppPurchase);
+	ProxyFactoryClass = UInAppPurchaseCallbackProxy::StaticClass();
+
+	ProxyClass = UInAppPurchaseCallbackProxy::StaticClass();
+}
+
+#undef LOCTEXT_NAMESPACE
diff --git a/Engine/Plugins/Online/OnlineSubsystemUtils/Source/OnlineBlueprintSupport/Private/K2Node_InAppPurchaseQuery.cpp b/Engine/Plugins/Online/OnlineSubsystemUtils/Source/OnlineBlueprintSupport/Private/K2Node_InAppPurchaseQuery.cpp
new file mode 100644
--- /dev/null
+++ b/Engine/Plugins/Online/OnlineSubsystemUtils/Source/OnlineBlueprintSupport/Private/K2Node_InAppPurchaseQuery.cpp
@@ -0,0 +1,18 @@
+// Copyright 1998-2016 Epic Games, Inc. All Rights Reserved.
+
+#include "OnlineBlueprintSupportPrivatePCH.h"
+#include "InAppPurchaseQueryCallbackProxy.h"
+#include "K2Node_InAppPurchaseQuery.h"
+
+#define LOCTEXT_NAMESPACE "K2Node"
+
+UK2Node_InAppPurchaseQuery::UK2Node_InAppPurchaseQuery(const FObjectInitializer& ObjectInitializer)
+	: Super(ObjectInitializer)
+{
+	ProxyFactoryFunctionName = GET_FUNCTION_NAME_CHECKED(UInAppPurchaseQueryCallbackProxy, CreateProxyObjectForInAppPurchaseQuery);
+	ProxyFactoryClass = UInAppPurchaseQueryCallbackProxy::StaticClass();
+
+	ProxyClass = UInAppPurchaseQueryCallbackProxy::StaticClass();
+}
+
+#undef LOCTEXT_NAMESPACE
diff --git a/Engine/Plugins/Online/OnlineSubsystemUtils/Source/OnlineBlueprintSupport/Private/OnlineBlueprintSupportModule.cpp b/Engine/Plugins/Online/OnlineSubsystemUtils/Source/OnlineBlueprintSupport/Private/OnlineBlueprintSupportModule.cpp
new file mode 100644
--- /dev/null
+++ b/Engine/Plugins/Online/OnlineSubsystemUtils/Source/OnlineBlueprintSupport/Private/OnlineBlueprintSupportModule.cpp
@@ -0,0 +1,7 @@
+// Copyright 1998-2016 Epic Games, Inc. All Rights Reserved.
+
+#include "OnlineBlueprintSupportPrivatePCH.h"
+#include "ModuleManager.h"
+
+
+IMPLEMENT_MODULE( FDefaultModuleImpl, OnlineBlueprintSupport );
diff --git a/Engine/Plugins/Online/OnlineSubsystemUtils/Source/OnlineBlueprintSupport/Private/OnlineBlueprintSupportPrivatePCH.h b/Engine/Plugins/Online/OnlineSubsystemUtils/Source/OnlineBlueprintSupport/Private/OnlineBlueprintSupportPrivatePCH.h
new file mode 100644
--- /dev/null
+++ b/Engine/Plugins/Online/OnlineSubsystemUtils/Source/OnlineBlueprintSupport/Private/OnlineBlueprintSupportPrivatePCH.h
@@ -0,0 +1,15 @@
+// Copyright 1998-2016 Epic Games, Inc. All Rights Reserved.
+
+#pragma once
+
+
+#include "Engine.h"
+#include "OnlineSubsystem.h"
+#include "OnlineSubsystemModule.h"
+#include "SocketSubsystem.h"
+#include "BlueprintGraphDefinitions.h"
+//#include "BlueprintUtilities.h"
+//#include "Kismet2/BlueprintEditorUtils.h"
+//#include "UnrealEd.h"
+//#include "EditorStyle.h"
+
diff --git a/Engine/Plugins/Online/OnlineSubsystemUtils/Source/OnlineSubsystemUtils/Classes/IpConnection.h b/Engine/Plugins/Online/OnlineSubsystemUtils/Source/OnlineSubsystemUtils/Classes/IpConnection.h
new file mode 100644
--- /dev/null
+++ b/Engine/Plugins/Online/OnlineSubsystemUtils/Source/OnlineSubsystemUtils/Classes/IpConnection.h
@@ -0,0 +1,30 @@
+// Copyright 1998-2016 Epic Games, Inc. All Rights Reserved.
+
+//
+// Ip based implementation of a network connection used by the net driver class
+//
+
+#pragma once
+#include "IpConnection.generated.h"
+
+UCLASS(transient, config=Engine)
+class ONLINESUBSYSTEMUTILS_API UIpConnection : public UNetConnection
+{
+    GENERATED_UCLASS_BODY()
+	// Variables.
+	TSharedPtr<FInternetAddr>	RemoteAddr;
+	class FSocket*				Socket;
+	class FResolveInfo*			ResolveInfo;
+
+	//~ Begin NetConnection Interface
+	virtual void InitBase(UNetDriver* InDriver, class FSocket* InSocket, const FURL& InURL, EConnectionState InState, int32 InMaxPacket = 0, int32 InPacketOverhead = 0) override;
+	virtual void InitRemoteConnection(UNetDriver* InDriver, class FSocket* InSocket, const FURL& InURL, const class FInternetAddr& InRemoteAddr, EConnectionState InState, int32 InMaxPacket = 0, int32 InPacketOverhead = 0) override;
+	virtual void InitLocalConnection(UNetDriver* InDriver, class FSocket* InSocket, const FURL& InURL, EConnectionState InState, int32 InMaxPacket = 0, int32 InPacketOverhead = 0) override;
+	virtual void LowLevelSend(void* Data, int32 CountBytes, int32 CountBits) override;
+	FString LowLevelGetRemoteAddress(bool bAppendPort=false) override;
+	FString LowLevelDescribe() override;
+	virtual int32 GetAddrAsInt(void) override;
+	virtual int32 GetAddrPort(void) override;
+	virtual FString RemoteAddressToString() override;
+	//~ End NetConnection Interface
+};
diff --git a/Engine/Plugins/Online/OnlineSubsystemUtils/Source/OnlineSubsystemUtils/Classes/IpNetDriver.h b/Engine/Plugins/Online/OnlineSubsystemUtils/Source/OnlineSubsystemUtils/Classes/IpNetDriver.h
new file mode 100644
--- /dev/null
+++ b/Engine/Plugins/Online/OnlineSubsystemUtils/Source/OnlineSubsystemUtils/Classes/IpNetDriver.h
@@ -0,0 +1,79 @@
+// Copyright 1998-2016 Epic Games, Inc. All Rights Reserved.
+
+//
+// Ip endpoint based implementation of the net driver
+//
+
+#pragma once
+#include "IpNetDriver.generated.h"
+
+UCLASS(transient, config=Engine)
+class ONLINESUBSYSTEMUTILS_API UIpNetDriver : public UNetDriver
+{
+    GENERATED_UCLASS_BODY()
+
+	/** Should port unreachable messages be logged */
+    UPROPERTY(Config)
+    uint32 LogPortUnreach:1;
+
+	/** Does the game allow clients to remain after receiving ICMP port unreachable errors (handles flakey connections) */
+    UPROPERTY(Config)
+    uint32 AllowPlayerPortUnreach:1;
+
+	/** Number of ports which will be tried if current one is not available for binding (i.e. if told to bind to port N, will try from N to N+MaxPortCountToTry inclusive) */
+	UPROPERTY(Config)
+	uint32 MaxPortCountToTry;
+
+	/** Local address this net driver is associated with */
+	TSharedPtr<FInternetAddr> LocalAddr;
+
+	/** Underlying socket communication */
+	FSocket* Socket;
+
+	//~ Begin UNetDriver Interface.
+	virtual bool IsAvailable() const override;
+	virtual bool InitBase(bool bInitAsClient, FNetworkNotify* InNotify, const FURL& URL, bool bReuseAddressAndPort, FString& Error) override;
+	virtual bool InitConnect( FNetworkNotify* InNotify, const FURL& ConnectURL, FString& Error ) override;
+	virtual bool InitListen( FNetworkNotify* InNotify, FURL& LocalURL, bool bReuseAddressAndPort, FString& Error ) override;
+	virtual void ProcessRemoteFunction(class AActor* Actor, class UFunction* Function, void* Parameters, struct FOutParmRec* OutParms, struct FFrame* Stack, class UObject* SubObject = NULL) override;
+	virtual void TickDispatch( float DeltaTime ) override;
+	virtual void LowLevelSend(FString Address, void* Data, int32 CountBits) override;
+	virtual FString LowLevelGetNetworkNumber() override;
+	virtual void LowLevelDestroy() override;
+	virtual class ISocketSubsystem* GetSocketSubsystem() override;
+	virtual bool IsNetResourceValid(void) override
+	{
+		return Socket != NULL;
+	}
+	//~ End UNetDriver Interface
+
+	//~ Begin UIpNetDriver Interface.
+	virtual FSocket * CreateSocket();
+
+	/**
+	 * Returns the port number to use when a client is creating a socket.
+	 * Platforms that can't use the default of 0 (system-selected port) may override
+	 * this function.
+	 *
+	 * @return The port number to use for client sockets. Base implementation returns 0.
+	 */
+	virtual int GetClientPort();
+	//~ End UIpNetDriver Interface.
+
+	//~ Begin FExec Interface
+	virtual bool Exec( UWorld* InWorld, const TCHAR* Cmd, FOutputDevice& Ar=*GLog ) override;
+	//~ End FExec Interface
+
+	/**
+	 * Exec command handlers
+	 */
+	bool HandleSocketsCommand( const TCHAR* Cmd, FOutputDevice& Ar, UWorld* InWorld );
+
+	/** @return TCPIP connection to server */
+	class UIpConnection* GetServerConnection();
+
+	// Callback for platform handling when networking is taking a long time in a single frame (by default over 1 second).
+	// It may get called multiple times in a single frame if additional processing after a previous alert exceeds the threshold again
+	DECLARE_MULTICAST_DELEGATE(FOnNetworkProcessingCausingSlowFrame);
+	static FOnNetworkProcessingCausingSlowFrame OnNetworkProcessingCausingSlowFrame;
+};
diff --git a/Engine/Plugins/Online/OnlineSubsystemUtils/Source/OnlineSubsystemUtils/Classes/LeaderboardBlueprintLibrary.h b/Engine/Plugins/Online/OnlineSubsystemUtils/Source/OnlineSubsystemUtils/Classes/LeaderboardBlueprintLibrary.h
new file mode 100644
--- /dev/null
+++ b/Engine/Plugins/Online/OnlineSubsystemUtils/Source/OnlineSubsystemUtils/Classes/LeaderboardBlueprintLibrary.h
@@ -0,0 +1,21 @@
+// Copyright 1998-2016 Epic Games, Inc. All Rights Reserved.
+
+#pragma once
+
+#include "LeaderboardBlueprintLibrary.generated.h"
+
+/**
+ * A beacon host used for taking reservations for an existing game session
+ */
+UCLASS()
+class ONLINESUBSYSTEMUTILS_API ULeaderboardBlueprintLibrary : public UBlueprintFunctionLibrary
+{
+	GENERATED_UCLASS_BODY()
+
+	/** Writes an integer value to the specified leaderboard */
+	UFUNCTION(BlueprintCallable, Category = "Online|Leaderboard")
+	static bool WriteLeaderboardInteger(APlayerController* PlayerController, FName StatName, int32 StatValue);
+
+private:
+	static bool WriteLeaderboardObject(APlayerController* PlayerController, class FOnlineLeaderboardWrite& WriteObject);
+};
diff --git a/Engine/Plugins/Online/OnlineSubsystemUtils/Source/OnlineSubsystemUtils/Classes/TestBeaconClient.h b/Engine/Plugins/Online/OnlineSubsystemUtils/Source/OnlineSubsystemUtils/Classes/TestBeaconClient.h
new file mode 100644
--- /dev/null
+++ b/Engine/Plugins/Online/OnlineSubsystemUtils/Source/OnlineSubsystemUtils/Classes/TestBeaconClient.h
@@ -0,0 +1,26 @@
+// Copyright 1998-2016 Epic Games, Inc. All Rights Reserved.
+
+#pragma once
+#include "OnlineBeaconClient.h"
+#include "TestBeaconClient.generated.h"
+
+/**
+ * A beacon client used for making reservations with an existing game session
+ */
+UCLASS(transient, notplaceable, config=Engine)
+class ONLINESUBSYSTEMUTILS_API ATestBeaconClient : public AOnlineBeaconClient
+{
+	GENERATED_UCLASS_BODY()
+
+	//~ Begin AOnlineBeaconClient Interface
+	virtual void OnFailure() override;
+	//~ End AOnlineBeaconClient Interface
+
+	/** Send a ping RPC to the client */
+	UFUNCTION(client, reliable)
+	virtual void ClientPing();
+
+	/** Send a pong RPC to the host */
+	UFUNCTION(server, reliable, WithValidation)
+	virtual void ServerPong();
+};
diff --git a/Engine/Plugins/Online/OnlineSubsystemUtils/Source/OnlineSubsystemUtils/Classes/TestBeaconHost.h b/Engine/Plugins/Online/OnlineSubsystemUtils/Source/OnlineSubsystemUtils/Classes/TestBeaconHost.h
new file mode 100644
--- /dev/null
+++ b/Engine/Plugins/Online/OnlineSubsystemUtils/Source/OnlineSubsystemUtils/Classes/TestBeaconHost.h
@@ -0,0 +1,21 @@
+// Copyright 1998-2016 Epic Games, Inc. All Rights Reserved.
+
+#pragma once
+#include "OnlineBeaconHostObject.h"
+#include "TestBeaconHost.generated.h"
+
+/**
+ * A beacon host used for taking reservations for an existing game session
+ */
+UCLASS(transient, notplaceable, config=Engine)
+class ONLINESUBSYSTEMUTILS_API ATestBeaconHost : public AOnlineBeaconHostObject
+{
+	GENERATED_UCLASS_BODY()
+
+	//~ Begin AOnlineBeaconHost Interface 
+	virtual AOnlineBeaconClient* SpawnBeaconActor(class UNetConnection* ClientConnection) override;
+	virtual void OnClientConnected(class AOnlineBeaconClient* NewClientActor, class UNetConnection* ClientConnection) override;
+	//~ End AOnlineBeaconHost Interface 
+
+	virtual bool Init();
+};
diff --git a/Engine/Plugins/Online/OnlineSubsystemUtils/Source/OnlineSubsystemUtils/OnlineSubsystemUtils.Build.cs b/Engine/Plugins/Online/OnlineSubsystemUtils/Source/OnlineSubsystemUtils/OnlineSubsystemUtils.Build.cs
new file mode 100644
--- /dev/null
+++ b/Engine/Plugins/Online/OnlineSubsystemUtils/Source/OnlineSubsystemUtils/OnlineSubsystemUtils.Build.cs
@@ -0,0 +1,29 @@
+// Copyright 1998-2016 Epic Games, Inc. All Rights Reserved.
+
+using UnrealBuildTool;
+
+public class OnlineSubsystemUtils : ModuleRules
+{
+	public OnlineSubsystemUtils(TargetInfo Target)
+	{
+		Definitions.Add("ONLINESUBSYSTEMUTILS_PACKAGE=1");
+
+        PrivateIncludePaths.Add("OnlineSubsystemUtils/Private");
+
+        PrivateDependencyModuleNames.AddRange(
+			new string[] { 
+				"Core", 
+				"CoreUObject",
+				"Engine", 
+				"EngineSettings",
+                "ImageCore",
+				"Sockets",
+				"Voice",
+                "PacketHandler",
+				"Json"
+			}
+		);
+
+        PublicDependencyModuleNames.Add("OnlineSubsystem");
+	}
+}
diff --git a/Engine/Plugins/Online/OnlineSubsystemUtils/Source/OnlineSubsystemUtils/Private/InAppPurchaseCallbackProxy.cpp b/Engine/Plugins/Online/OnlineSubsystemUtils/Source/OnlineSubsystemUtils/Private/InAppPurchaseCallbackProxy.cpp
new file mode 100644
--- /dev/null
+++ b/Engine/Plugins/Online/OnlineSubsystemUtils/Source/OnlineSubsystemUtils/Private/InAppPurchaseCallbackProxy.cpp
@@ -0,0 +1,143 @@
+// Copyright 1998-2016 Epic Games, Inc. All Rights Reserved.
+
+#include "OnlineSubsystemUtilsPrivatePCH.h"
+#include "InAppPurchaseCallbackProxy.h"
+
+//////////////////////////////////////////////////////////////////////////
+// UInAppPurchaseCallbackProxy
+
+UInAppPurchaseCallbackProxy::UInAppPurchaseCallbackProxy(const FObjectInitializer& ObjectInitializer)
+	: Super(ObjectInitializer)
+{
+	PurchaseRequest = nullptr;
+	WorldPtr = nullptr;
+}
+
+
+void UInAppPurchaseCallbackProxy::Trigger(APlayerController* PlayerController, const FInAppPurchaseProductRequest& ProductRequest)
+{
+	bFailedToEvenSubmit = true;
+
+	WorldPtr = (PlayerController != nullptr) ? PlayerController->GetWorld() : nullptr;
+	if (APlayerState* PlayerState = (PlayerController != nullptr) ? PlayerController->PlayerState : nullptr)
+	{
+		if (IOnlineSubsystem* const OnlineSub = IOnlineSubsystem::IsLoaded() ? IOnlineSubsystem::Get() : nullptr)
+		{
+			IOnlineStorePtr StoreInterface = OnlineSub->GetStoreInterface();
+			if (StoreInterface.IsValid())
+			{
+				bFailedToEvenSubmit = false;
+
+				// Register the completion callback
+				InAppPurchaseCompleteDelegate = FOnInAppPurchaseCompleteDelegate::CreateUObject(this, &UInAppPurchaseCallbackProxy::OnInAppPurchaseComplete);
+				InAppPurchaseCompleteDelegateHandle = StoreInterface->AddOnInAppPurchaseCompleteDelegate_Handle(InAppPurchaseCompleteDelegate);
+
+				// Set-up, and trigger the transaction through the store interface
+				PurchaseRequest = MakeShareable(new FOnlineInAppPurchaseTransaction());
+				FOnlineInAppPurchaseTransactionRef PurchaseRequestRef = PurchaseRequest.ToSharedRef();
+				StoreInterface->BeginPurchase(ProductRequest, PurchaseRequestRef);
+			}
+			else
+			{
+				FFrame::KismetExecutionMessage(TEXT("UInAppPurchaseCallbackProxy::Trigger - In-App Purchases are not supported by Online Subsystem"), ELogVerbosity::Warning);
+			}
+		}
+		else
+		{
+			FFrame::KismetExecutionMessage(TEXT("UInAppPurchaseCallbackProxy::Trigger - Invalid or uninitialized OnlineSubsystem"), ELogVerbosity::Warning);
+		}
+	}
+	else
+	{
+		FFrame::KismetExecutionMessage(TEXT("UInAppPurchaseCallbackProxy::Trigger - Invalid player state"), ELogVerbosity::Warning);
+	}
+
+	if (bFailedToEvenSubmit && (PlayerController != NULL))
+	{
+		OnInAppPurchaseComplete(EInAppPurchaseState::Failed);
+	}
+}
+
+
+void UInAppPurchaseCallbackProxy::OnInAppPurchaseComplete(EInAppPurchaseState::Type CompletionState)
+{
+	RemoveDelegate();
+	SavedPurchaseState = CompletionState;
+    
+	if (UWorld* World = WorldPtr.Get())
+	{
+		DECLARE_CYCLE_STAT(TEXT("FSimpleDelegateGraphTask.DelayInAppPurchaseComplete"), STAT_FSimpleDelegateGraphTask_DelayInAppPurchaseComplete, STATGROUP_TaskGraphTasks);
+
+		FSimpleDelegateGraphTask::CreateAndDispatchWhenReady(
+			FSimpleDelegateGraphTask::FDelegate::CreateLambda([=](){
+
+				OnInAppPurchaseComplete_Delayed();
+
+			}),
+			GET_STATID(STAT_FSimpleDelegateGraphTask_DelayInAppPurchaseComplete), 
+			nullptr, 
+			ENamedThreads::GameThread
+		);
+    }
+    else
+    {
+        PurchaseRequest = nullptr;
+    }
+}
+
+
+void UInAppPurchaseCallbackProxy::OnInAppPurchaseComplete_Delayed()
+{
+    /** Cached product details of the purchased product */
+    FInAppPurchaseProductInfo ProductInformation; 
+
+    if (SavedPurchaseState == EInAppPurchaseState::Success && PurchaseRequest.IsValid())
+    {
+        ProductInformation = PurchaseRequest->ProvidedProductInformation;
+    }
+    
+	if (SavedPurchaseState == EInAppPurchaseState::Success)
+	{
+		OnSuccess.Broadcast(SavedPurchaseState, ProductInformation);
+	}
+	else
+	{
+		OnFailure.Broadcast(SavedPurchaseState, ProductInformation);
+	}
+    
+    PurchaseRequest = nullptr;
+}
+
+
+void UInAppPurchaseCallbackProxy::RemoveDelegate()
+{
+	if (!bFailedToEvenSubmit)
+	{
+		if (IOnlineSubsystem* OnlineSub = IOnlineSubsystem::IsLoaded() ? IOnlineSubsystem::Get() : nullptr)
+		{
+			IOnlineStorePtr InAppPurchases = OnlineSub->GetStoreInterface();
+			if (InAppPurchases.IsValid())
+			{
+				InAppPurchases->ClearOnInAppPurchaseCompleteDelegate_Handle(InAppPurchaseCompleteDelegateHandle);
+			}
+		}
+	}
+}
+
+
+void UInAppPurchaseCallbackProxy::BeginDestroy()
+{
+	PurchaseRequest = nullptr;
+	RemoveDelegate();
+
+	Super::BeginDestroy();
+}
+
+
+UInAppPurchaseCallbackProxy* UInAppPurchaseCallbackProxy::CreateProxyObjectForInAppPurchase(class APlayerController* PlayerController, const FInAppPurchaseProductRequest& ProductRequest)
+{
+	UInAppPurchaseCallbackProxy* Proxy = NewObject<UInAppPurchaseCallbackProxy>();
+	Proxy->SetFlags(RF_StrongRefOnFrame);
+	Proxy->Trigger(PlayerController, ProductRequest);
+	return Proxy;
+}
diff --git a/Engine/Plugins/Online/OnlineSubsystemUtils/Source/OnlineSubsystemUtils/Private/InAppPurchaseQueryCallbackProxy.cpp b/Engine/Plugins/Online/OnlineSubsystemUtils/Source/OnlineSubsystemUtils/Private/InAppPurchaseQueryCallbackProxy.cpp
new file mode 100644
--- /dev/null
+++ b/Engine/Plugins/Online/OnlineSubsystemUtils/Source/OnlineSubsystemUtils/Private/InAppPurchaseQueryCallbackProxy.cpp
@@ -0,0 +1,132 @@
+// Copyright 1998-2016 Epic Games, Inc. All Rights Reserved.
+
+#include "OnlineSubsystemUtilsPrivatePCH.h"
+#include "InAppPurchaseQueryCallbackProxy.h"
+
+//////////////////////////////////////////////////////////////////////////
+// UInAppPurchaseQueryCallbackProxy
+
+UInAppPurchaseQueryCallbackProxy::UInAppPurchaseQueryCallbackProxy(const FObjectInitializer& ObjectInitializer)
+	: Super(ObjectInitializer)
+{
+}
+
+void UInAppPurchaseQueryCallbackProxy::TriggerQuery(APlayerController* PlayerController, const TArray<FString>& ProductIdentifiers)
+{
+	bFailedToEvenSubmit = true;
+
+	WorldPtr = (PlayerController != nullptr) ? PlayerController->GetWorld() : nullptr;
+	if (APlayerState* PlayerState = (PlayerController != NULL) ? PlayerController->PlayerState : nullptr)
+	{
+		if (IOnlineSubsystem* const OnlineSub = IOnlineSubsystem::IsLoaded() ? IOnlineSubsystem::Get() : nullptr)
+		{
+			IOnlineStorePtr StoreInterface = OnlineSub->GetStoreInterface();
+			if (StoreInterface.IsValid())
+			{
+				bFailedToEvenSubmit = false;
+
+				// Register the completion callback
+				InAppPurchaseReadCompleteDelegate       = FOnQueryForAvailablePurchasesCompleteDelegate::CreateUObject(this, &UInAppPurchaseQueryCallbackProxy::OnInAppPurchaseRead);
+				InAppPurchaseReadCompleteDelegateHandle = StoreInterface->AddOnQueryForAvailablePurchasesCompleteDelegate_Handle(InAppPurchaseReadCompleteDelegate);
+
+
+				ReadObject = MakeShareable(new FOnlineProductInformationRead());
+				FOnlineProductInformationReadRef ReadObjectRef = ReadObject.ToSharedRef();
+				StoreInterface->QueryForAvailablePurchases(ProductIdentifiers, ReadObjectRef);
+			}
+			else
+			{
+				FFrame::KismetExecutionMessage(TEXT("UInAppPurchaseQueryCallbackProxy::TriggerQuery - In App Purchases are not supported by Online Subsystem"), ELogVerbosity::Warning);
+			}
+		}
+		else
+		{
+			FFrame::KismetExecutionMessage(TEXT("UInAppPurchaseQueryCallbackProxy::TriggerQuery - Invalid or uninitialized OnlineSubsystem"), ELogVerbosity::Warning);
+		}
+	}
+	else
+	{
+		FFrame::KismetExecutionMessage(TEXT("UInAppPurchaseQueryCallbackProxy::TriggerQuery - Invalid player state"), ELogVerbosity::Warning);
+	}
+
+	if (bFailedToEvenSubmit && (PlayerController != NULL))
+	{
+		OnInAppPurchaseRead(false);
+	}
+}
+
+void UInAppPurchaseQueryCallbackProxy::OnInAppPurchaseRead(bool bWasSuccessful)
+{
+	RemoveDelegate();
+
+	if (bWasSuccessful && ReadObject.IsValid())
+	{
+		SavedProductInformation = ReadObject->ProvidedProductInformation;
+		bSavedWasSuccessful = true;
+	}
+	else
+	{
+		bSavedWasSuccessful = false;
+	}
+
+	if (UWorld* World = WorldPtr.Get())
+	{
+		DECLARE_CYCLE_STAT(TEXT("FSimpleDelegateGraphTask.DelayInAppPurchaseRead"), STAT_FSimpleDelegateGraphTask_DelayInAppPurchaseRead, STATGROUP_TaskGraphTasks);
+
+		FSimpleDelegateGraphTask::CreateAndDispatchWhenReady(
+			FSimpleDelegateGraphTask::FDelegate::CreateLambda([=](){
+
+				OnInAppPurchaseRead_Delayed();
+
+			}),
+			GET_STATID(STAT_FSimpleDelegateGraphTask_DelayInAppPurchaseRead), 
+			nullptr, 
+			ENamedThreads::GameThread
+		);
+	}
+
+	ReadObject = NULL;
+}
+
+void UInAppPurchaseQueryCallbackProxy::OnInAppPurchaseRead_Delayed()
+{
+	if (bSavedWasSuccessful)
+	{
+		OnSuccess.Broadcast(SavedProductInformation);
+	}
+	else
+	{
+		OnFailure.Broadcast(SavedProductInformation);
+	}
+}
+
+void UInAppPurchaseQueryCallbackProxy::RemoveDelegate()
+{
+	if (!bFailedToEvenSubmit)
+	{
+		if (IOnlineSubsystem* OnlineSub = IOnlineSubsystem::IsLoaded() ? IOnlineSubsystem::Get() : nullptr)
+		{
+			IOnlineStorePtr InAppPurchases = OnlineSub->GetStoreInterface();
+			if (InAppPurchases.IsValid())
+			{
+				InAppPurchases->ClearOnQueryForAvailablePurchasesCompleteDelegate_Handle(InAppPurchaseReadCompleteDelegateHandle);
+			}
+		}
+	}
+}
+
+void UInAppPurchaseQueryCallbackProxy::BeginDestroy()
+{
+	ReadObject = NULL;
+	RemoveDelegate();
+
+	Super::BeginDestroy();
+}
+
+UInAppPurchaseQueryCallbackProxy* UInAppPurchaseQueryCallbackProxy::CreateProxyObjectForInAppPurchaseQuery(class APlayerController* PlayerController, const TArray<FString>& ProductIdentifiers)
+{
+	UInAppPurchaseQueryCallbackProxy* Proxy = NewObject<UInAppPurchaseQueryCallbackProxy>();
+	Proxy->SetFlags(RF_StrongRefOnFrame);
+	Proxy->TriggerQuery(PlayerController, ProductIdentifiers);
+	return Proxy;
+}
diff --git a/Engine/Plugins/Online/OnlineSubsystemUtils/Source/OnlineSubsystemUtils/Private/InAppPurchaseRestoreCallbackProxy.cpp b/Engine/Plugins/Online/OnlineSubsystemUtils/Source/OnlineSubsystemUtils/Private/InAppPurchaseRestoreCallbackProxy.cpp
new file mode 100644
--- /dev/null
+++ b/Engine/Plugins/Online/OnlineSubsystemUtils/Source/OnlineSubsystemUtils/Private/InAppPurchaseRestoreCallbackProxy.cpp
@@ -0,0 +1,134 @@
+// Copyright 1998-2016 Epic Games, Inc. All Rights Reserved.
+
+#include "OnlineSubsystemUtilsPrivatePCH.h"
+#include "InAppPurchaseRestoreCallbackProxy.h"
+
+//////////////////////////////////////////////////////////////////////////
+// UInAppPurchaseRestoreCallbackProxy
+
+UInAppPurchaseRestoreCallbackProxy::UInAppPurchaseRestoreCallbackProxy(const FObjectInitializer& ObjectInitializer)
+	: Super(ObjectInitializer)
+{
+	WorldPtr = nullptr;
+}
+
+
+void UInAppPurchaseRestoreCallbackProxy::Trigger(const TArray<FInAppPurchaseProductRequest>& ConsumableProductFlags, APlayerController* PlayerController)
+{
+	bFailedToEvenSubmit = true;
+
+	WorldPtr = (PlayerController != nullptr) ? PlayerController->GetWorld() : nullptr;
+	if (APlayerState* PlayerState = (PlayerController != nullptr) ? PlayerController->PlayerState : nullptr)
+	{
+		if (IOnlineSubsystem* const OnlineSub = IOnlineSubsystem::IsLoaded() ? IOnlineSubsystem::Get() : nullptr)
+		{
+			IOnlineStorePtr StoreInterface = OnlineSub->GetStoreInterface();
+			if (StoreInterface.IsValid())
+			{
+				bFailedToEvenSubmit = false;
+
+				// Register the completion callback
+				InAppPurchaseRestoreCompleteDelegate = FOnInAppPurchaseCompleteDelegate::CreateUObject(this, &UInAppPurchaseRestoreCallbackProxy::OnInAppPurchaseRestoreComplete);
+				InAppPurchaseRestoreCompleteDelegateHandle = StoreInterface->AddOnInAppPurchaseRestoreCompleteDelegate_Handle(InAppPurchaseRestoreCompleteDelegate);
+
+				// Set-up, and trigger the transaction through the store interface
+				ReadObject = MakeShareable(new FOnlineInAppPurchaseRestoreRead());
+				FOnlineInAppPurchaseRestoreReadRef ReadObjectRef = ReadObject.ToSharedRef();
+				StoreInterface->RestorePurchases(ConsumableProductFlags, ReadObjectRef);
+			}
+			else
+			{
+				FFrame::KismetExecutionMessage(TEXT("UInAppPurchaseRestoreCallbackProxy::Trigger - In-App Purchases are not supported by Online Subsystem"), ELogVerbosity::Warning);
+			}
+		}
+		else
+		{
+			FFrame::KismetExecutionMessage(TEXT("UInAppPurchaseRestoreCallbackProxy::Trigger - Invalid or uninitialized OnlineSubsystem"), ELogVerbosity::Warning);
+		}
+	}
+	else
+	{
+		FFrame::KismetExecutionMessage(TEXT("UInAppPurchaseRestoreCallbackProxy::Trigger - Invalid player state"), ELogVerbosity::Warning);
+	}
+
+	if (bFailedToEvenSubmit && (PlayerController != NULL))
+	{
+		OnInAppPurchaseRestoreComplete(EInAppPurchaseState::Failed);
+	}
+}
+
+
+void UInAppPurchaseRestoreCallbackProxy::OnInAppPurchaseRestoreComplete(EInAppPurchaseState::Type CompletionState)
+{
+	RemoveDelegate();
+	SavedPurchaseState = CompletionState;
+	if (CompletionState == EInAppPurchaseState::Restored)
+	{
+		SavedProductInformation = ReadObject->ProvidedRestoreInformation;
+	}
+    
+	if (UWorld* World = WorldPtr.Get())
+	{
+		DECLARE_CYCLE_STAT(TEXT("FSimpleDelegateGraphTask.DelayInAppPurchaseRestoreComplete"), STAT_FSimpleDelegateGraphTask_DelayInAppPurchaseRestoreComplete, STATGROUP_TaskGraphTasks);
+
+		FSimpleDelegateGraphTask::CreateAndDispatchWhenReady(
+			FSimpleDelegateGraphTask::FDelegate::CreateLambda([=](){
+
+				OnInAppPurchaseRestoreComplete_Delayed();
+
+			}),
+			GET_STATID(STAT_FSimpleDelegateGraphTask_DelayInAppPurchaseRestoreComplete), 
+			nullptr, 
+			ENamedThreads::GameThread
+		);
+    }
+
+	ReadObject = NULL;
+}
+
+
+void UInAppPurchaseRestoreCallbackProxy::OnInAppPurchaseRestoreComplete_Delayed()
+{    
+	if (SavedPurchaseState == EInAppPurchaseState::Restored)
+	{
+		OnSuccess.Broadcast(SavedPurchaseState, SavedProductInformation);
+	}
+	else
+	{
+		OnFailure.Broadcast(SavedPurchaseState, SavedProductInformation);
+	}
+}
+
+
+void UInAppPurchaseRestoreCallbackProxy::RemoveDelegate()
+{
+	if (!bFailedToEvenSubmit)
+	{
+		if (IOnlineSubsystem* OnlineSub = IOnlineSubsystem::IsLoaded() ? IOnlineSubsystem::Get() : nullptr)
+		{
+			IOnlineStorePtr InAppPurchases = OnlineSub->GetStoreInterface();
+			if (InAppPurchases.IsValid())
+			{
+				InAppPurchases->ClearOnInAppPurchaseRestoreCompleteDelegate_Handle(InAppPurchaseRestoreCompleteDelegateHandle);
+			}
+		}
+	}
+}
+
+
+void UInAppPurchaseRestoreCallbackProxy::BeginDestroy()
+{
+	ReadObject = NULL;
+	RemoveDelegate();
+
+	Super::BeginDestroy();
+}
+
+
+UInAppPurchaseRestoreCallbackProxy* UInAppPurchaseRestoreCallbackProxy::CreateProxyObjectForInAppPurchaseRestore(const TArray<FInAppPurchaseProductRequest>& ConsumableProductFlags, class APlayerController* PlayerController)
+{
+	UInAppPurchaseRestoreCallbackProxy* Proxy = NewObject<UInAppPurchaseRestoreCallbackProxy>();
+	Proxy->SetFlags(RF_StrongRefOnFrame);
+	Proxy->Trigger(ConsumableProductFlags, PlayerController);
+	return Proxy;
+}
diff --git a/Engine/Plugins/Online/OnlineSubsystemUtils/Source/OnlineSubsystemUtils/Private/IpConnection.cpp b/Engine/Plugins/Online/OnlineSubsystemUtils/Source/OnlineSubsystemUtils/Private/IpConnection.cpp
new file mode 100644
--- /dev/null
+++ b/Engine/Plugins/Online/OnlineSubsystemUtils/Source/OnlineSubsystemUtils/Private/IpConnection.cpp
@@ -0,0 +1,204 @@
+// Copyright 1998-2016 Epic Games, Inc. All Rights Reserved.
+
+/*=============================================================================
+	IpConnection.cpp: Unreal IP network connection.
+Notes:
+	* See \msdev\vc98\include\winsock.h and \msdev\vc98\include\winsock2.h 
+	  for Winsock WSAE* errors returned by Windows Sockets.
+=============================================================================*/
+
+#include "OnlineSubsystemUtilsPrivatePCH.h"
+
+#include "IPAddress.h"
+#include "Sockets.h"
+#include "Net/NetworkProfiler.h"
+#include "Net/DataChannel.h"
+#include "Runtime/PacketHandlers/PacketHandler/Public/PacketHandler.h"
+
+/*-----------------------------------------------------------------------------
+	Declarations.
+-----------------------------------------------------------------------------*/
+
+// Size of a UDP header.
+#define IP_HEADER_SIZE     (20)
+#define UDP_HEADER_SIZE    (IP_HEADER_SIZE+8)
+
+UIpConnection::UIpConnection(const FObjectInitializer& ObjectInitializer) :
+	Super(ObjectInitializer),
+	RemoteAddr(NULL),
+	Socket(NULL),
+	ResolveInfo(NULL)
+{
+}
+
+void UIpConnection::InitBase(UNetDriver* InDriver, class FSocket* InSocket, const FURL& InURL, EConnectionState InState, int32 InMaxPacket, int32 InPacketOverhead)
+{
+	// Pass the call up the chain
+	Super::InitBase(InDriver, InSocket, InURL, InState, 
+		// Use the default packet size/overhead unless overridden by a child class
+		(InMaxPacket == 0 || InMaxPacket > MAX_PACKET_SIZE) ? MAX_PACKET_SIZE : InMaxPacket,
+		InPacketOverhead == 0 ? UDP_HEADER_SIZE : InPacketOverhead);
+
+	Socket = InSocket;
+	ResolveInfo = NULL;
+}
+
+void UIpConnection::InitLocalConnection(UNetDriver* InDriver, class FSocket* InSocket, const FURL& InURL, EConnectionState InState, int32 InMaxPacket, int32 InPacketOverhead)
+{
+	InitBase(InDriver, InSocket, InURL, InState, 
+		// Use the default packet size/overhead unless overridden by a child class
+		(InMaxPacket == 0 || InMaxPacket > MAX_PACKET_SIZE) ? MAX_PACKET_SIZE : InMaxPacket,
+		InPacketOverhead == 0 ? UDP_HEADER_SIZE : InPacketOverhead);
+
+	// Figure out IP address from the host URL
+	bool bIsValid = false;
+	// Get numerical address directly.
+	RemoteAddr = InDriver->GetSocketSubsystem()->CreateInternetAddr();
+	RemoteAddr->SetIp(*InURL.Host, bIsValid);
+	RemoteAddr->SetPort(InURL.Port);
+
+	// Try to resolve it if it failed
+	if (bIsValid == false)
+	{
+		// Create thread to resolve the address.
+		ResolveInfo = InDriver->GetSocketSubsystem()->GetHostByName(TCHAR_TO_ANSI(*InURL.Host));
+		if (ResolveInfo == NULL)
+		{
+			Close();
+			UE_LOG(LogNet, Verbose, TEXT("IpConnection::InitConnection: Unable to resolve %s"), *InURL.Host);
+		}
+	}
+
+	// Initialize our send bunch
+	InitSendBuffer();
+}
+
+void UIpConnection::InitRemoteConnection(UNetDriver* InDriver, class FSocket* InSocket, const FURL& InURL, const class FInternetAddr& InRemoteAddr, EConnectionState InState, int32 InMaxPacket, int32 InPacketOverhead)
+{
+	InitBase(InDriver, InSocket, InURL, InState, 
+		// Use the default packet size/overhead unless overridden by a child class
+		(InMaxPacket == 0 || InMaxPacket > MAX_PACKET_SIZE) ? MAX_PACKET_SIZE : InMaxPacket,
+		InPacketOverhead == 0 ? UDP_HEADER_SIZE : InPacketOverhead);
+
+	// Copy the remote IPAddress passed in
+	bool bIsValid = false;
+	FString IpAddrStr = InRemoteAddr.ToString(false);
+	RemoteAddr = InDriver->GetSocketSubsystem()->CreateInternetAddr();
+	RemoteAddr->SetIp(*IpAddrStr, bIsValid);
+	RemoteAddr->SetPort(InRemoteAddr.GetPort());
+
+	URL.Host = RemoteAddr->ToString(false);
+
+	// Initialize our send bunch
+	InitSendBuffer();
+
+	// This is for a client that needs to log in, setup ClientLoginState and ExpectedClientLoginMsgType to reflect that
+	SetClientLoginState( EClientLoginState::LoggingIn );
+	SetExpectedClientLoginMsgType( NMT_Hello );
+}
+
+void UIpConnection::LowLevelSend(void* Data, int32 CountBytes, int32 CountBits)
+{
+	const uint8* DataToSend = reinterpret_cast<uint8*>(Data);
+
+	if( ResolveInfo )
+	{
+		// If destination address isn't resolved yet, send nowhere.
+		if( !ResolveInfo->IsComplete() )
+		{
+			// Host name still resolving.
+			return;
+		}
+		else if( ResolveInfo->GetErrorCode() != SE_NO_ERROR )
+		{
+			// Host name resolution just now failed.
+			UE_LOG(LogNet, Log,  TEXT("Host name resolution failed with %d"), ResolveInfo->GetErrorCode() );
+			Driver->ServerConnection->State = USOCK_Closed;
+			delete ResolveInfo;
+			ResolveInfo = NULL;
+			return;
+		}
+		else
+		{
+			uint32 Addr;
+			// Host name resolution just now succeeded.
+			ResolveInfo->GetResolvedAddress().GetIp(Addr);
+			RemoteAddr->SetIp(Addr);
+			UE_LOG(LogNet, Log, TEXT("Host name resolution completed"));
+			delete ResolveInfo;
+			ResolveInfo = NULL;
+		}
+	}
+
+
+	// Process any packet modifiers
+	if (Handler.IsValid() && !Handler->GetRawSend())
+	{
+		const ProcessedPacket ProcessedData = Handler->Outgoing(reinterpret_cast<uint8*>(Data), CountBits);
+
+		DataToSend = ProcessedData.Data;
+		CountBytes = FMath::DivideAndRoundUp(ProcessedData.CountBits, 8);
+		CountBits = ProcessedData.CountBits;
+	}
+
+	// Send to remote.
+	int32 BytesSent = 0;
+	CLOCK_CYCLES(Driver->SendCycles);
+
+	if ( CountBytes > MaxPacket )
+	{
+		UE_LOG( LogNet, Warning, TEXT( "UIpConnection::LowLevelSend: CountBytes > MaxPacketSize! Count: %i, MaxPacket: %i %s" ), CountBytes, MaxPacket, *Describe() );
+	}
+
+	if (CountBytes > 0)
+	{
+		Socket->SendTo(DataToSend, CountBytes, BytesSent, *RemoteAddr);
+	}
+
+	UNCLOCK_CYCLES(Driver->SendCycles);
+	NETWORK_PROFILER(GNetworkProfiler.FlushOutgoingBunches(this));
+	NETWORK_PROFILER(GNetworkProfiler.TrackSocketSendTo(Socket->GetDescription(),DataToSend,BytesSent,NumPacketIdBits,NumBunchBits,NumAckBits,NumPaddingBits,this));
+}
+
+FString UIpConnection::LowLevelGetRemoteAddress(bool bAppendPort)
+{
+	return RemoteAddr->ToString(bAppendPort);
+}
+
+FString UIpConnection::LowLevelDescribe()
+{
+	TSharedRef<FInternetAddr> LocalAddr = Driver->GetSocketSubsystem()->CreateInternetAddr();
+	Socket->GetAddress(*LocalAddr);
+	return FString::Printf
+	(
+		TEXT("url=%s remote=%s local=%s state: %s"),
+		*URL.Host,
+		*RemoteAddr->ToString(true),
+		*LocalAddr->ToString(true),
+			State==USOCK_Pending	?	TEXT("Pending")
+		:	State==USOCK_Open		?	TEXT("Open")
+		:	State==USOCK_Closed		?	TEXT("Closed")
+		:								TEXT("Invalid")
+	);
+}
+
+int32 UIpConnection::GetAddrAsInt(void)
+{
+	uint32 OutAddr = 0;
+	// Get the host byte order ip addr
+	RemoteAddr->GetIp(OutAddr);
+	return (int32)OutAddr;
+}
+
+int32 UIpConnection::GetAddrPort(void)
+{
+	int32 OutPort = 0;
+	// Get the host byte order ip port
+	RemoteAddr->GetPort(OutPort);
+	return OutPort;
+}
+
+FString UIpConnection::RemoteAddressToString()
+{
+	return RemoteAddr->ToString(true);
+}
\ No newline at end of file
diff --git a/Engine/Plugins/Online/OnlineSubsystemUtils/Source/OnlineSubsystemUtils/Private/IpNetDriver.cpp b/Engine/Plugins/Online/OnlineSubsystemUtils/Source/OnlineSubsystemUtils/Private/IpNetDriver.cpp
new file mode 100644
--- /dev/null
+++ b/Engine/Plugins/Online/OnlineSubsystemUtils/Source/OnlineSubsystemUtils/Private/IpNetDriver.cpp
@@ -0,0 +1,572 @@
+// Copyright 1998-2016 Epic Games, Inc. All Rights Reserved.
+
+/*=============================================================================
+	IpNetDriver.cpp: Unreal IP network driver.
+Notes:
+	* See \msdev\vc98\include\winsock.h and \msdev\vc98\include\winsock2.h 
+	  for Winsock WSAE* errors returned by Windows Sockets.
+=============================================================================*/
+
+#include "OnlineSubsystemUtilsPrivatePCH.h"
+#include "Engine/Channel.h"
+
+#include "IPAddress.h"
+#include "Sockets.h"
+
+/*-----------------------------------------------------------------------------
+	Declarations.
+-----------------------------------------------------------------------------*/
+
+UIpNetDriver::FOnNetworkProcessingCausingSlowFrame UIpNetDriver::OnNetworkProcessingCausingSlowFrame;
+
+// Time before the alarm delegate is called (in seconds)
+float GIpNetDriverMaxDesiredTimeSliceBeforeAlarmSecs = 1.0f;
+
+FAutoConsoleVariableRef GIpNetDriverMaxDesiredTimeSliceBeforeAlarmSecsCVar(
+	TEXT("n.IpNetDriverMaxFrameTimeBeforeAlert"),
+	GIpNetDriverMaxDesiredTimeSliceBeforeAlarmSecs,
+	TEXT("Time to spend processing networking data in a single frame before an alert is raised (in seconds)\n")
+	TEXT("It may get called multiple times in a single frame if additional processing after a previous alert exceeds the threshold again\n")
+	TEXT(" default: 1 s"));
+
+// Time before the time taken in a single frame is printed out (in seconds)
+float GIpNetDriverLongFramePrintoutThresholdSecs = 10.0f;
+
+FAutoConsoleVariableRef GIpNetDriverLongFramePrintoutThresholdSecsCVar(
+	TEXT("n.IpNetDriverMaxFrameTimeBeforeLogging"),
+	GIpNetDriverLongFramePrintoutThresholdSecs,
+	TEXT("Time to spend processing networking data in a single frame before an output log warning is printed (in seconds)\n")
+	TEXT(" default: 10 s"));
+
+UIpNetDriver::UIpNetDriver(const FObjectInitializer& ObjectInitializer)
+	: Super(ObjectInitializer)
+{
+}
+
+bool UIpNetDriver::IsAvailable() const
+{
+	// IP driver always valid for now
+	return true;
+}
+
+ISocketSubsystem* UIpNetDriver::GetSocketSubsystem()
+{
+	return ISocketSubsystem::Get();
+}
+
+FSocket * UIpNetDriver::CreateSocket()
+{
+	// Create UDP socket and enable broadcasting.
+	ISocketSubsystem* SocketSubsystem = GetSocketSubsystem();
+
+	if (SocketSubsystem == NULL)
+	{
+		UE_LOG(LogNet, Warning, TEXT("UIpNetDriver::CreateSocket: Unable to find socket subsystem"));
+		return NULL;
+	}
+
+	return SocketSubsystem->CreateSocket( NAME_DGram, TEXT( "Unreal" ) );
+}
+
+int UIpNetDriver::GetClientPort()
+{
+	return 0;
+}
+
+bool UIpNetDriver::InitBase( bool bInitAsClient, FNetworkNotify* InNotify, const FURL& URL, bool bReuseAddressAndPort, FString& Error )
+{
+	if (!Super::InitBase(bInitAsClient, InNotify, URL, bReuseAddressAndPort, Error))
+	{
+		return false;
+	}
+
+	ISocketSubsystem* SocketSubsystem = GetSocketSubsystem();
+	if (SocketSubsystem == NULL)
+	{
+		UE_LOG(LogNet, Warning, TEXT("Unable to find socket subsystem"));
+		return false;
+	}
+
+	// Derived types may have already allocated a socket
+
+	// Create the socket that we will use to communicate with
+	Socket = CreateSocket();
+
+	if( Socket == NULL )
+	{
+		Socket = 0;
+		Error = FString::Printf( TEXT("WinSock: socket failed (%i)"), (int32)SocketSubsystem->GetLastErrorCode() );
+		return false;
+	}
+	if (SocketSubsystem->RequiresChatDataBeSeparate() == false &&
+		Socket->SetBroadcast() == false)
+	{
+		Error = FString::Printf( TEXT("%s: setsockopt SO_BROADCAST failed (%i)"), SocketSubsystem->GetSocketAPIName(), (int32)SocketSubsystem->GetLastErrorCode() );
+		return false;
+	}
+
+	if (Socket->SetReuseAddr(bReuseAddressAndPort) == false)
+	{
+		UE_LOG(LogNet, Log, TEXT("setsockopt with SO_REUSEADDR failed"));
+	}
+
+	if (Socket->SetRecvErr() == false)
+	{
+		UE_LOG(LogNet, Log, TEXT("setsockopt with IP_RECVERR failed"));
+	}
+
+	// Increase socket queue size, because we are polling rather than threading
+	// and thus we rely on Windows Sockets to buffer a lot of data on the server.
+	int32 RecvSize = bInitAsClient ? 0x8000 : 0x20000;
+	int32 SendSize = bInitAsClient ? 0x8000 : 0x20000;
+	Socket->SetReceiveBufferSize(RecvSize,RecvSize);
+	Socket->SetSendBufferSize(SendSize,SendSize);
+	UE_LOG(LogInit, Log, TEXT("%s: Socket queue %i / %i"), SocketSubsystem->GetSocketAPIName(), RecvSize, SendSize );
+
+	// Bind socket to our port.
+	LocalAddr = SocketSubsystem->GetLocalBindAddr(*GLog);
+	
+	LocalAddr->SetPort(bInitAsClient ? GetClientPort() : URL.Port);
+	
+	int32 AttemptPort = LocalAddr->GetPort();
+	int32 BoundPort = SocketSubsystem->BindNextPort( Socket, *LocalAddr, MaxPortCountToTry + 1, 1 );
+	if (BoundPort == 0)
+	{
+		Error = FString::Printf( TEXT("%s: binding to port %i failed (%i)"), SocketSubsystem->GetSocketAPIName(), AttemptPort,
+			(int32)SocketSubsystem->GetLastErrorCode() );
+		return false;
+	}
+	if( Socket->SetNonBlocking() == false )
+	{
+		Error = FString::Printf( TEXT("%s: SetNonBlocking failed (%i)"), SocketSubsystem->GetSocketAPIName(),
+			(int32)SocketSubsystem->GetLastErrorCode());
+		return false;
+	}
+
+	// Success.
+	return true;
+}
+
+bool UIpNetDriver::InitConnect( FNetworkNotify* InNotify, const FURL& ConnectURL, FString& Error )
+{
+	if( !InitBase( true, InNotify, ConnectURL, false, Error ) )
+	{
+		UE_LOG(LogNet, Warning, TEXT("Failed to init net driver ConnectURL: %s: %s"), *ConnectURL.ToString(), *Error);
+		return false;
+	}
+
+	// Create new connection.
+	ServerConnection = NewObject<UNetConnection>(GetTransientPackage(), NetConnectionClass);
+	ServerConnection->InitLocalConnection( this, Socket, ConnectURL, USOCK_Pending);
+	UE_LOG(LogNet, Log, TEXT("Game client on port %i, rate %i"), ConnectURL.Port, ServerConnection->CurrentNetSpeed );
+
+	// Create channel zero.
+	GetServerConnection()->CreateChannel( CHTYPE_Control, 1, 0 );
+
+	return true;
+}
+
+bool UIpNetDriver::InitListen( FNetworkNotify* InNotify, FURL& LocalURL, bool bReuseAddressAndPort, FString& Error )
+{
+	if( !InitBase( false, InNotify, LocalURL, bReuseAddressAndPort, Error ) )
+	{
+		UE_LOG(LogNet, Warning, TEXT("Failed to init net driver ListenURL: %s: %s"), *LocalURL.ToString(), *Error);
+		return false;
+	}
+
+
+	InitConnectionlessHandler();
+
+	// Update result URL.
+	//LocalURL.Host = LocalAddr->ToString(false);
+	LocalURL.Port = LocalAddr->GetPort();
+	UE_LOG(LogNet, Log, TEXT("%s IpNetDriver listening on port %i"), *GetDescription(), LocalURL.Port );
+
+	return true;
+}
+
+void UIpNetDriver::TickDispatch( float DeltaTime )
+{
+	Super::TickDispatch( DeltaTime );
+
+	ISocketSubsystem* SocketSubsystem = GetSocketSubsystem();
+
+	const double StartReceiveTime = FPlatformTime::Seconds();
+	double AlarmTime = StartReceiveTime + GIpNetDriverMaxDesiredTimeSliceBeforeAlarmSecs;
+
+	// Process all incoming packets.
+	uint8 Data[MAX_PACKET_SIZE];
+	uint8* DataRef = Data;
+	bool bIgnorePacket = false;
+	TSharedRef<FInternetAddr> FromAddr = SocketSubsystem->CreateInternetAddr();
+
+	for( ; Socket != NULL; )
+	{
+		{
+			const double CurrentTime = FPlatformTime::Seconds();
+			if (CurrentTime > AlarmTime)
+			{
+				OnNetworkProcessingCausingSlowFrame.Broadcast();
+
+				AlarmTime = CurrentTime + GIpNetDriverMaxDesiredTimeSliceBeforeAlarmSecs;
+			}
+		}
+
+		int32 BytesRead = 0;
+
+		// Get data, if any.
+		CLOCK_CYCLES(RecvCycles);
+		bool bOk = Socket->RecvFrom(Data, sizeof(Data), BytesRead, *FromAddr);
+		UNCLOCK_CYCLES(RecvCycles);
+
+		if (bOk)
+		{
+			// Immediately stop processing, for empty packets (usually a DDoS)
+			if (BytesRead == 0)
+			{
+				break;
+			}
+		}
+		else
+		{
+			ESocketErrors Error = SocketSubsystem->GetLastErrorCode();
+			if(Error == SE_EWOULDBLOCK ||
+			   Error == SE_NO_ERROR)
+			{
+				// No data or no error?
+				break;
+			}
+			else
+			{
+				// MalformedPacket: Client tried sending a packet that exceeded the maximum packet limit
+				// enforced by the server
+				if (Error == SE_EMSGSIZE)
+				{
+					UIpConnection* Connection = nullptr;
+					if (GetServerConnection() && (*GetServerConnection()->RemoteAddr == *FromAddr))
+					{
+						Connection = GetServerConnection();
+					}
+
+					if (Connection != nullptr)
+					{
+						UE_SECURITY_LOG(Connection, ESecurityEvent::Malformed_Packet, TEXT("Received Packet with bytes > max MTU"));
+					}
+				}
+
+				if( Error != SE_ECONNRESET && Error != SE_UDP_ERR_PORT_UNREACH )
+				{
+					UE_LOG(LogNet, Warning, TEXT("UDP recvfrom error: %i (%s) from %s"),
+						(int32)Error,
+						SocketSubsystem->GetSocketError(Error),
+						*FromAddr->ToString(true));
+					break;
+				}
+			}
+		}
+		// Figure out which socket the received data came from.
+		UIpConnection* Connection = nullptr;
+		UIpConnection* MyServerConnection = GetServerConnection();
+		if (MyServerConnection)
+		{
+			if ((*MyServerConnection->RemoteAddr == *FromAddr))
+			{
+				Connection = MyServerConnection;
+			}
+			else
+			{
+				UE_LOG(LogNet, Warning, TEXT("Incoming ip address doesn't match expected server address: Actual: %s Expected: %s"),
+					*FromAddr->ToString(true),
+					MyServerConnection->RemoteAddr.IsValid() ? *MyServerConnection->RemoteAddr->ToString(true) : TEXT("Invalid"));
+			}
+		}
+		for( int32 i=0; i<ClientConnections.Num() && !Connection; i++ )
+		{
+			UIpConnection* TestConnection = (UIpConnection*)ClientConnections[i]; 
+            check(TestConnection);
+			if(*TestConnection->RemoteAddr == *FromAddr)
+			{
+				Connection = TestConnection;
+			}
+		}
+
+		if( bOk == false )
+		{
+			if( Connection )
+			{
+				if( Connection != GetServerConnection() )
+				{
+					// We received an ICMP port unreachable from the client, meaning the client is no longer running the game
+					// (or someone is trying to perform a DoS attack on the client)
+
+					// rcg08182002 Some buggy firewalls get occasional ICMP port
+					// unreachable messages from legitimate players. Still, this code
+					// will drop them unceremoniously, so there's an option in the .INI
+					// file for servers with such flakey connections to let these
+					// players slide...which means if the client's game crashes, they
+					// might get flooded to some degree with packets until they timeout.
+					// Either way, this should close up the usual DoS attacks.
+					if ((Connection->State != USOCK_Open) || (!AllowPlayerPortUnreach))
+					{
+						if (LogPortUnreach)
+						{
+							UE_LOG(LogNet, Log, TEXT("Received ICMP port unreachable from client %s.  Disconnecting."),
+								*FromAddr->ToString(true));
+						}
+						Connection->CleanUp();
+					}
+				}
+			}
+			else
+			{
+				if (LogPortUnreach)
+				{
+					UE_LOG(LogNet, Log, TEXT("Received ICMP port unreachable from %s.  No matching connection found."),
+						*FromAddr->ToString(true));
+				}
+			}
+		}
+		else
+		{
+			// If we didn't find a client connection, maybe create a new one.
+			if( !Connection )
+			{
+				// Determine if allowing for client/server connections
+				const bool bAcceptingConnection = Notify != nullptr && Notify->NotifyAcceptingConnection() == EAcceptConnection::Accept;
+
+				if (bAcceptingConnection)
+				{
+					bool bPassedChallenge = false;
+
+					bIgnorePacket = true;
+
+					if (ConnectionlessHandler.IsValid() && StatelessConnectComponent.IsValid())
+					{
+						TSharedPtr<StatelessConnectHandlerComponent> StatelessConnect = StatelessConnectComponent.Pin();
+						FString IncomingAddress = FromAddr->ToString(true);
+
+						const ProcessedPacket UnProcessedPacket =
+												ConnectionlessHandler->IncomingConnectionless(IncomingAddress, DataRef, BytesRead);
+
+						bPassedChallenge = StatelessConnect->HasPassedChallenge(IncomingAddress);
+
+						if (bPassedChallenge)
+						{
+							BytesRead = FMath::DivideAndRoundUp(UnProcessedPacket.CountBits, 8);
+
+							if (BytesRead > 0)
+							{
+								DataRef = UnProcessedPacket.Data;
+								bIgnorePacket = false;
+							}
+						}
+					}
+#if !UE_BUILD_SHIPPING
+					else if (FParse::Param(FCommandLine::Get(), TEXT("NoPacketHandler")))
+					{
+						UE_LOG(LogNet, Log, TEXT("Accepting connection without handshake, due to '-NoPacketHandler'."))
+
+						bIgnorePacket = false;
+						bPassedChallenge = true;
+					}
+#endif
+					else
+					{
+						UE_LOG(LogNet, Log,
+								TEXT("Invalid ConnectionlessHandler (%i) or StatelessConnectComponent (%i); can't accept connections."),
+								(int32)(ConnectionlessHandler.IsValid()), (int32)(StatelessConnectComponent.IsValid()));
+					}
+
+					if (bPassedChallenge)
+					{
+						UE_LOG(LogNet, Log, TEXT("Server accepting post-challenge connection from: %s"), *FromAddr->ToString(true));
+
+						Connection = NewObject<UIpConnection>(GetTransientPackage(), NetConnectionClass);
+						check(Connection);
+						Connection->InitRemoteConnection( this, Socket,  FURL(), *FromAddr, USOCK_Open);
+						Notify->NotifyAcceptedConnection( Connection );
+						AddClientConnection(Connection);
+					}
+				}
+			}
+
+			// Send the packet to the connection for processing.
+			if (Connection && !bIgnorePacket)
+			{
+				Connection->ReceivedRawPacket( DataRef, BytesRead );
+			}
+		}
+	}
+
+	const double EndReceiveTime = FPlatformTime::Seconds();
+	const float DeltaReceiveTime = EndReceiveTime - StartReceiveTime;
+
+	if (DeltaReceiveTime > GIpNetDriverLongFramePrintoutThresholdSecs)
+	{
+		UE_LOG( LogNet, Warning, TEXT( "UIpNetDriver::TickDispatch: Took too long to receive packets. Time: %2.2f %s" ), DeltaReceiveTime, *GetName() );
+	}
+}
+
+void UIpNetDriver::LowLevelSend(FString Address, void* Data, int32 CountBits)
+{
+	bool bValidAddress = !Address.IsEmpty();
+	TSharedRef<FInternetAddr> RemoteAddr = GetSocketSubsystem()->CreateInternetAddr();
+
+	if (bValidAddress)
+	{
+		RemoteAddr->SetIp(*Address, bValidAddress);
+	}
+
+	if (bValidAddress)
+	{
+		const uint8* DataToSend = reinterpret_cast<uint8*>(Data);
+
+		if (ConnectionlessHandler.IsValid())
+		{
+			const ProcessedPacket ProcessedData =
+					ConnectionlessHandler->OutgoingConnectionless(Address, (uint8*)DataToSend, CountBits);
+
+			DataToSend = ProcessedData.Data;
+			CountBits = ProcessedData.CountBits;
+		}
+
+
+		int32 BytesSent = 0;
+
+		CLOCK_CYCLES(SendCycles);
+		Socket->SendTo(DataToSend, FMath::DivideAndRoundUp(CountBits, 8), BytesSent, *RemoteAddr);
+		UNCLOCK_CYCLES(SendCycles);
+
+
+		// @todo: Can't implement these profiling events (require UNetConnections)
+		//NETWORK_PROFILER(GNetworkProfiler.FlushOutgoingBunches(/* UNetConnection */));
+		//NETWORK_PROFILER(GNetworkProfiler.TrackSocketSendTo(Socket->GetDescription(),Data,BytesSent,NumPacketIdBits,NumBunchBits,
+							//NumAckBits,NumPaddingBits, /* UNetConnection */));
+	}
+	else
+	{
+		UE_LOG(LogNet, Warning, TEXT("UIpNetDriver::LowLevelSend: Invalid send address '%s'"), *Address);
+	}
+}
+
+void UIpNetDriver::ProcessRemoteFunction(class AActor* Actor, UFunction* Function, void* Parameters, FOutParmRec* OutParms, FFrame* Stack, class UObject* SubObject )
+{
+	bool bIsServer = IsServer();
+
+	UNetConnection* Connection = NULL;
+	if (bIsServer)
+	{
+		if ((Function->FunctionFlags & FUNC_NetMulticast))
+		{
+			// Multicast functions go to every client
+			TArray<UNetConnection*> UniqueRealConnections;
+			for (int32 i=0; i<ClientConnections.Num(); ++i)
+			{
+				Connection = ClientConnections[i];
+				if (Connection && Connection->ViewTarget)
+				{
+					// Do relevancy check if unreliable.
+					// Reliables will always go out. This is odd behavior. On one hand we wish to guarantee "reliables always get there". On the other
+					// hand, replicating a reliable to something on the other side of the map that is non relevant seems weird. 
+					//
+					// Multicast reliables should probably never be used in gameplay code for actors that have relevancy checks. If they are, the 
+					// rpc will go through and the channel will be closed soon after due to relevancy failing.
+
+					bool IsRelevant = true;
+					if ((Function->FunctionFlags & FUNC_NetReliable) == 0)
+					{
+						FNetViewer Viewer(Connection, 0.f);
+						IsRelevant = Actor->IsNetRelevantFor(Viewer.InViewer, Viewer.ViewTarget, Viewer.ViewLocation);
+					}
+					
+					if (IsRelevant)
+					{
+						if (Connection->GetUChildConnection() != NULL)
+						{
+							Connection = ((UChildConnection*)Connection)->Parent;
+						}
+						
+						InternalProcessRemoteFunction( Actor, SubObject, Connection, Function, Parameters, OutParms, Stack, bIsServer );
+					}
+				}
+			}			
+
+			// Replicate any RPCs to the replay net driver so that they can get saved in network replays
+			UNetDriver* NetDriver = GEngine->FindNamedNetDriver(GetWorld(), NAME_DemoNetDriver);
+			if (NetDriver)
+			{
+				NetDriver->ProcessRemoteFunction(Actor, Function, Parameters, OutParms, Stack, SubObject);
+			}
+			// Return here so we don't call InternalProcessRemoteFunction again at the bottom of this function
+			return;
+		}
+	}
+
+	// Send function data to remote.
+	Connection = Actor->GetNetConnection();
+	if (Connection)
+	{
+		InternalProcessRemoteFunction( Actor, SubObject, Connection, Function, Parameters, OutParms, Stack, bIsServer );
+	}
+	else
+	{
+		UE_LOG(LogNet, Warning, TEXT("UIpNetDriver::ProcesRemoteFunction: No owning connection for actor %s. Function %s will not be processed."), *Actor->GetName(), *Function->GetName());
+	}
+}
+
+FString UIpNetDriver::LowLevelGetNetworkNumber()
+{
+	return LocalAddr->ToString(true);
+}
+
+void UIpNetDriver::LowLevelDestroy()
+{
+	Super::LowLevelDestroy();
+
+	// Close the socket.
+	if( Socket && !HasAnyFlags(RF_ClassDefaultObject) )
+	{
+		ISocketSubsystem* SocketSubsystem = GetSocketSubsystem();
+		if( !Socket->Close() )
+		{
+			UE_LOG(LogExit, Log, TEXT("closesocket error (%i)"), (int32)SocketSubsystem->GetLastErrorCode() );
+		}
+		// Free the memory the OS allocated for this socket
+		SocketSubsystem->DestroySocket(Socket);
+		Socket = NULL;
+		UE_LOG(LogExit, Log, TEXT("%s shut down"),*GetDescription() );
+	}
+
+}
+
+
+bool UIpNetDriver::HandleSocketsCommand( const TCHAR* Cmd, FOutputDevice& Ar, UWorld* InWorld )
+{
+	Ar.Logf(TEXT(""));
+	if (Socket != NULL)
+	{
+		TSharedRef<FInternetAddr> LocalInternetAddr = GetSocketSubsystem()->CreateInternetAddr();
+		Socket->GetAddress(*LocalInternetAddr);
+		Ar.Logf(TEXT("%s Socket: %s"), *GetDescription(), *LocalInternetAddr->ToString(true));
+	}		
+	else
+	{
+		Ar.Logf(TEXT("%s Socket: null"), *GetDescription());
+	}
+	return UNetDriver::Exec( InWorld, TEXT("SOCKETS"),Ar);
+}
+
+bool UIpNetDriver::Exec( UWorld* InWorld, const TCHAR* Cmd, FOutputDevice& Ar )
+{
+	if (FParse::Command(&Cmd,TEXT("SOCKETS")))
+	{
+		return HandleSocketsCommand( Cmd, Ar, InWorld );
+	}
+	return UNetDriver::Exec( InWorld, Cmd,Ar);
+}
+
+UIpConnection* UIpNetDriver::GetServerConnection() 
+{
+	return (UIpConnection*)ServerConnection;
+}
+
diff --git a/Engine/Plugins/Online/OnlineSubsystemUtils/Source/OnlineSubsystemUtils/Private/OnlineBeacon.cpp b/Engine/Plugins/Online/OnlineSubsystemUtils/Source/OnlineSubsystemUtils/Private/OnlineBeacon.cpp
new file mode 100644
--- /dev/null
+++ b/Engine/Plugins/Online/OnlineSubsystemUtils/Source/OnlineSubsystemUtils/Private/OnlineBeacon.cpp
@@ -0,0 +1,183 @@
+// Copyright 1998-2016 Epic Games, Inc. All Rights Reserved.
+
+#include "OnlineSubsystemUtilsPrivatePCH.h"
+#include "OnlineBeacon.h"
+#include "Net/DataChannel.h"
+
+DEFINE_LOG_CATEGORY(LogBeacon);
+
+AOnlineBeacon::AOnlineBeacon(const FObjectInitializer& ObjectInitializer) :
+	Super(ObjectInitializer),
+	NetDriver(nullptr),
+	BeaconState(EBeaconState::DenyRequests)
+{
+	NetDriverName = FName(TEXT("BeaconDriver"));
+}
+
+bool AOnlineBeacon::InitBase()
+{
+	NetDriver = GEngine->CreateNetDriver(GetWorld(), NAME_BeaconNetDriver);
+	if (NetDriver != nullptr)
+	{
+		HandleNetworkFailureDelegateHandle = GEngine->OnNetworkFailure().AddUObject(this, &AOnlineBeacon::HandleNetworkFailure);
+		SetNetDriverName(NetDriver->NetDriverName);
+		return true;
+	}
+
+	return false;
+}
+
+void AOnlineBeacon::EndPlay(const EEndPlayReason::Type EndPlayReason)
+{
+	if (NetDriver)
+	{
+		GEngine->DestroyNamedNetDriver(GetWorld(), NetDriverName);
+		NetDriver = nullptr;
+	}
+
+	Super::EndPlay(EndPlayReason);
+}
+
+bool AOnlineBeacon::HasNetOwner() const
+{
+    // Beacons are their own net owners
+	return true;
+}
+
+void AOnlineBeacon::DestroyBeacon()
+{
+	UE_LOG(LogBeacon, Verbose, TEXT("Destroying beacon %s, netdriver %s"), *GetName(), NetDriver ? *NetDriver->GetDescription() : TEXT("NULL"));
+	GEngine->OnNetworkFailure().Remove(HandleNetworkFailureDelegateHandle);
+
+	if (NetDriver)
+	{
+		GEngine->DestroyNamedNetDriver(GetWorld(), NetDriverName);
+		NetDriver = nullptr;
+	}
+
+	Destroy();
+}
+
+void AOnlineBeacon::HandleNetworkFailure(UWorld *World, UNetDriver *InNetDriver, ENetworkFailure::Type FailureType, const FString& ErrorString)
+{
+	if (InNetDriver && InNetDriver->NetDriverName == NetDriverName)
+	{
+		UE_LOG(LogBeacon, Verbose, TEXT("NetworkFailure %s: %s"), *GetName(), ENetworkFailure::ToString(FailureType));
+		OnFailure();
+	}
+}
+
+void AOnlineBeacon::OnFailure()
+{
+	GEngine->OnNetworkFailure().Remove(HandleNetworkFailureDelegateHandle);
+	
+	if (NetDriver)
+	{
+		GEngine->DestroyNamedNetDriver(GetWorld(), NetDriverName);
+		NetDriver = nullptr;
+	}
+}
+
+void AOnlineBeacon::OnActorChannelOpen(FInBunch& Bunch, UNetConnection* Connection)
+{
+	Connection->OwningActor = this;
+	Super::OnActorChannelOpen(Bunch, Connection);
+}
+
+bool AOnlineBeacon::IsRelevancyOwnerFor(const AActor* ReplicatedActor, const AActor* ActorOwner, const AActor* ConnectionActor) const
+{
+	bool bRelevantOwner = (ConnectionActor == ReplicatedActor);
+	return bRelevantOwner;
+}
+
+bool AOnlineBeacon::IsNetRelevantFor(const AActor* RealViewer, const AActor* ViewTarget, const FVector& SrcLocation) const
+{
+	// Only replicate to the owner or to connections of the same beacon type (possible that multiple UNetConnections come from the same client)
+	bool bIsOwner = GetNetConnection() == ViewTarget->GetNetConnection();
+	bool bSameBeaconType = GetClass() == RealViewer->GetClass();
+	return bOnlyRelevantToOwner ? bIsOwner : bSameBeaconType;
+}
+
+EAcceptConnection::Type AOnlineBeacon::NotifyAcceptingConnection()
+{
+	check(NetDriver);
+	if(NetDriver->ServerConnection)
+	{
+		// We are a client and we don't welcome incoming connections.
+		UE_LOG(LogNet, Log, TEXT("NotifyAcceptingConnection: Client refused"));
+		return EAcceptConnection::Reject;
+	}
+	else if(BeaconState == EBeaconState::DenyRequests)
+	{
+		// Server is down
+		UE_LOG(LogNet, Log, TEXT("NotifyAcceptingConnection: Server %s refused"), *GetName());
+		return EAcceptConnection::Reject;
+	}
+	else //if(BeaconState == EBeaconState::AllowRequests)
+	{
+		// Server is up and running.
+		UE_LOG(LogNet, Log, TEXT("NotifyAcceptingConnection: Server %s accept"), *GetName());
+		return EAcceptConnection::Accept;
+	}
+}
+
+void AOnlineBeacon::NotifyAcceptedConnection(UNetConnection* Connection)
+{
+	check(NetDriver != nullptr);
+	check(NetDriver->ServerConnection == nullptr);
+	UE_LOG(LogNet, Log, TEXT("NotifyAcceptedConnection: Name: %s, TimeStamp: %s, %s"), *GetName(), FPlatformTime::StrTimestamp(), *Connection->Describe());
+}
+
+bool AOnlineBeacon::NotifyAcceptingChannel(UChannel* Channel)
+{
+	check(Channel);
+	check(Channel->Connection);
+	check(Channel->Connection->Driver);
+	UNetDriver* Driver = Channel->Connection->Driver;
+	check(NetDriver == Driver);
+
+	if (Driver->ServerConnection)
+	{
+		// We are a client and the server has just opened up a new channel.
+		UE_LOG(LogNet, Log,  TEXT("NotifyAcceptingChannel %i/%i client %s"), Channel->ChIndex, (int32)Channel->ChType, *GetName());
+		if (Channel->ChType == CHTYPE_Actor)
+		{
+			// Actor channel.
+			UE_LOG(LogNet, Log,  TEXT("Client accepting actor channel"));
+			return 1;
+		}
+		else if (Channel->ChType == CHTYPE_Voice)
+		{
+			// Accept server requests to open a voice channel, allowing for custom voip implementations
+			// which utilize multiple server controlled voice channels.
+			UE_LOG(LogNet, Log,  TEXT("Client accepting voice channel"));
+			return 1;
+		}
+		else
+		{
+			// Unwanted channel type.
+			UE_LOG(LogNet, Log, TEXT("Client refusing unwanted channel of type %i"), (uint8)Channel->ChType);
+			return 0;
+		}
+	}
+	else
+	{
+		// We are the server.
+		if (Channel->ChIndex == 0 && Channel->ChType == CHTYPE_Control)
+		{
+			// The client has opened initial channel.
+			UE_LOG(LogNet, Log, TEXT("NotifyAcceptingChannel Control %i server %s: Accepted"), Channel->ChIndex, *GetFullName());
+			return 1;
+		}
+		else
+		{
+			// Client can't open any other kinds of channels.
+			UE_LOG(LogNet, Log, TEXT("NotifyAcceptingChannel %i %i server %s: Refused"), (uint8)Channel->ChType, Channel->ChIndex, *GetFullName());
+			return 0;
+		}
+	}
+}
+
+void AOnlineBeacon::NotifyControlMessage(UNetConnection* Connection, uint8 MessageType, FInBunch& Bunch)
+{
+}
\ No newline at end of file
diff --git a/Engine/Plugins/Online/OnlineSubsystemUtils/Source/OnlineSubsystemUtils/Private/OnlineBeaconClient.cpp b/Engine/Plugins/Online/OnlineSubsystemUtils/Source/OnlineSubsystemUtils/Private/OnlineBeaconClient.cpp
new file mode 100644
--- /dev/null
+++ b/Engine/Plugins/Online/OnlineSubsystemUtils/Source/OnlineSubsystemUtils/Private/OnlineBeaconClient.cpp
@@ -0,0 +1,287 @@
+// Copyright 1998-2016 Epic Games, Inc. All Rights Reserved.
+
+#include "OnlineSubsystemUtilsPrivatePCH.h"
+#include "OnlineBeaconClient.h"
+#include "OnlineBeaconHost.h"
+#include "OnlineBeaconHostObject.h"
+#include "Net/DataChannel.h"
+#include "Net/UnrealNetwork.h"
+#include "NetworkVersion.h"
+
+#define BEACON_RPC_TIMEOUT 15.0f
+
+AOnlineBeaconClient::AOnlineBeaconClient(const FObjectInitializer& ObjectInitializer) :
+	Super(ObjectInitializer),
+	BeaconOwner(nullptr),
+	BeaconConnection(nullptr),
+	ConnectionState(EBeaconConnectionState::Invalid)
+{
+	NetDriverName = FName(TEXT("BeaconDriverClient"));
+	bOnlyRelevantToOwner = true;
+}
+
+FString AOnlineBeaconClient::GetBeaconType() const
+{
+	return GetClass()->GetName();
+}
+
+AOnlineBeaconHostObject* AOnlineBeaconClient::GetBeaconOwner() const
+{
+	return BeaconOwner;
+}
+
+void AOnlineBeaconClient::SetBeaconOwner(AOnlineBeaconHostObject* InBeaconOwner)
+{
+	BeaconOwner = InBeaconOwner;
+}
+
+const AActor* AOnlineBeaconClient::GetNetOwner() const
+{
+	return BeaconOwner;
+}
+
+class UNetConnection* AOnlineBeaconClient::GetNetConnection() const
+{
+	return BeaconConnection;
+}
+
+bool AOnlineBeaconClient::DestroyNetworkActorHandled()
+{
+	if (BeaconConnection)
+	{
+		BeaconConnection->bPendingDestroy = true;
+		return true;
+	}
+
+	return false;
+}
+
+const FUniqueNetIdRepl& AOnlineBeaconClient::GetUniqueId() const
+{
+	if (BeaconConnection)
+	{
+		return BeaconConnection->PlayerId;
+	}
+
+	static FUniqueNetIdRepl EmptyId;
+	return EmptyId;
+}
+
+EBeaconConnectionState AOnlineBeaconClient::GetConnectionState() const
+{
+	return ConnectionState;
+}
+
+void AOnlineBeaconClient::SetConnectionState(EBeaconConnectionState NewConnectionState)
+{
+	ConnectionState = NewConnectionState;
+}
+
+bool AOnlineBeaconClient::InitClient(FURL& URL)
+{
+	bool bSuccess = false;
+
+	if(URL.Valid)
+	{
+		if (InitBase() && NetDriver)
+		{
+			FString Error;
+			if (NetDriver->InitConnect(this, URL, Error))
+			{
+				BeaconConnection = NetDriver->ServerConnection;
+
+				// Kick off the connection handshake
+				if (BeaconConnection->StatelessConnectComponent.IsValid())
+				{
+					BeaconConnection->StatelessConnectComponent.Pin()->SendInitialConnect();
+				}
+
+				SetConnectionState(EBeaconConnectionState::Pending);
+
+				NetDriver->SetWorld(GetWorld());
+				NetDriver->Notify = this;
+				NetDriver->InitialConnectTimeout = BeaconConnectionInitialTimeout;
+				NetDriver->ConnectionTimeout = BeaconConnectionTimeout;
+
+				// Send initial message.
+				uint8 IsLittleEndian = uint8(PLATFORM_LITTLE_ENDIAN);
+				check(IsLittleEndian == !!IsLittleEndian); // should only be one or zero
+
+				uint32 LocalNetworkVersion = FNetworkVersion::GetLocalNetworkVersion();
+				
+				FNetControlMessage<NMT_Hello>::Send(NetDriver->ServerConnection, IsLittleEndian, LocalNetworkVersion);
+				NetDriver->ServerConnection->FlushNet();
+
+				bSuccess = true;
+			}
+			else
+			{
+				// error initializing the network stack...
+				UE_LOG(LogBeacon, Log, TEXT("AOnlineBeaconClient::InitClient failed"));
+				SetConnectionState(EBeaconConnectionState::Invalid);
+				OnFailure();
+			}
+		}
+	}
+
+	return bSuccess;
+}
+
+void AOnlineBeaconClient::OnFailure()
+{
+	UE_LOG(LogBeacon, Verbose, TEXT("Client beacon (%s) connection failure, handling connection timeout."), *GetName());
+	SetConnectionState(EBeaconConnectionState::Invalid);
+	HostConnectionFailure.ExecuteIfBound();
+	Super::OnFailure();
+}
+
+void AOnlineBeaconClient::ClientOnConnected_Implementation()
+{
+	SetConnectionState(EBeaconConnectionState::Open);
+	BeaconConnection->State = USOCK_Open;
+
+	Role = ROLE_Authority;
+	SetReplicates(true);
+	SetAutonomousProxy(true);
+
+	// Fail safe for connection to server but no client connection RPC
+	GetWorldTimerManager().ClearTimer(TimerHandle_OnFailure);
+
+	// Call the overloaded function for this client class
+	OnConnected();
+}
+
+bool AOnlineBeaconClient::UseShortConnectTimeout() const
+{
+	return ConnectionState == EBeaconConnectionState::Open;
+}
+
+void AOnlineBeaconClient::DestroyBeacon()
+{
+	SetConnectionState(EBeaconConnectionState::Closed);
+
+	UWorld* World = GetWorld();
+	if (World)
+	{
+		// Fail safe for connection to server but no client connection RPC
+		GetWorldTimerManager().ClearTimer(TimerHandle_OnFailure);
+	}
+
+	Super::DestroyBeacon();
+}
+
+void AOnlineBeaconClient::OnNetCleanup(UNetConnection* Connection)
+{
+	ensure(Connection == BeaconConnection);
+	SetConnectionState(EBeaconConnectionState::Closed);
+
+	AOnlineBeaconHostObject* BeaconHostObject = GetBeaconOwner();
+	if (BeaconHostObject)
+	{
+		BeaconHostObject->NotifyClientDisconnected(this);
+	}
+
+	BeaconConnection = nullptr;
+	Destroy(true);
+}
+
+void AOnlineBeaconClient::NotifyControlMessage(UNetConnection* Connection, uint8 MessageType, class FInBunch& Bunch)
+{
+	if(NetDriver->ServerConnection)
+	{
+		check(Connection == NetDriver->ServerConnection);
+
+		// We are the client
+#if !(UE_BUILD_SHIPPING && WITH_EDITOR)
+		UE_LOG(LogBeacon, Log, TEXT("%s[%s] Client received: %s"), *GetName(), *Connection->GetName(), FNetControlMessageInfo::GetName(MessageType));
+#endif
+		switch (MessageType)
+		{
+		case NMT_BeaconWelcome:
+			{
+				Connection->ClientResponse = TEXT("0");
+				FNetControlMessage<NMT_Netspeed>::Send(Connection, Connection->CurrentNetSpeed);
+
+				FString BeaconType = GetBeaconType();
+				if (!BeaconType.IsEmpty())
+				{
+					FUniqueNetIdRepl UniqueIdRepl;
+
+					ULocalPlayer* LocalPlayer = GEngine->GetFirstGamePlayer(GetWorld());
+					if (LocalPlayer)
+					{
+						// Send the player unique Id at login
+						UniqueIdRepl = LocalPlayer->GetPreferredUniqueNetId();
+					}
+
+					FNetControlMessage<NMT_BeaconJoin>::Send(Connection, BeaconType, UniqueIdRepl);
+					NetDriver->ServerConnection->FlushNet();
+				}
+				else
+				{
+					// Force close the session
+					UE_LOG(LogBeacon, Log, TEXT("Beacon close from invalid beacon type"));
+					OnFailure();
+				}
+				break;
+			}
+		case NMT_BeaconAssignGUID:
+			{
+				FNetworkGUID NetGUID;
+				FNetControlMessage<NMT_BeaconAssignGUID>::Receive(Bunch, NetGUID);
+				if (NetGUID.IsValid())
+				{
+					Connection->Driver->GuidCache->RegisterNetGUID_Client( NetGUID, this );
+
+					FString BeaconType = GetBeaconType();
+					FNetControlMessage<NMT_BeaconNetGUIDAck>::Send(Connection, BeaconType);
+					// Server will send ClientOnConnected() when it gets this control message
+
+					// Fail safe for connection to server but no client connection RPC
+					FTimerDelegate TimerDelegate = FTimerDelegate::CreateUObject(this, &AOnlineBeaconClient::OnFailure);
+					GetWorldTimerManager().SetTimer(TimerHandle_OnFailure, TimerDelegate, BEACON_RPC_TIMEOUT, false);
+				}
+				else
+				{
+					// Force close the session
+					UE_LOG(LogBeacon, Log, TEXT("Beacon close from invalid NetGUID"));
+					OnFailure();
+				}
+				break;
+			}
+		case NMT_Upgrade:
+			{
+				// Report mismatch.
+				uint32 RemoteNetworkVersion;
+				FNetControlMessage<NMT_Upgrade>::Receive(Bunch, RemoteNetworkVersion);
+				// Upgrade
+				const FString ConnectionError = NSLOCTEXT("Engine", "ClientOutdated", "The match you are trying to join is running an incompatible version of the game.  Please try upgrading your game version.").ToString();
+				GEngine->BroadcastNetworkFailure(GetWorld(), NetDriver, ENetworkFailure::OutdatedClient, ConnectionError);
+				break;
+			}
+		case NMT_Failure:
+			{
+				FString ErrorMsg;
+				FNetControlMessage<NMT_Failure>::Receive(Bunch, ErrorMsg);
+				if (ErrorMsg.IsEmpty())
+				{
+					ErrorMsg = NSLOCTEXT("NetworkErrors", "GenericBeaconConnectionFailed", "Beacon Connection Failed.").ToString();
+				}
+
+				// Force close the session
+				UE_LOG(LogBeacon, Log, TEXT("Beacon close from NMT_Failure %s"), *ErrorMsg);
+				OnFailure();
+				break;
+			}
+		case NMT_BeaconJoin:
+		case NMT_BeaconNetGUIDAck:
+		default:
+			{
+				// Force close the session
+				UE_LOG(LogBeacon, Log, TEXT("Beacon close from unexpected control message"));
+				OnFailure();
+				break;
+			}
+		}
+	}	
+}
diff --git a/Engine/Plugins/Online/OnlineSubsystemUtils/Source/OnlineSubsystemUtils/Private/OnlineBeaconHost.cpp b/Engine/Plugins/Online/OnlineSubsystemUtils/Source/OnlineSubsystemUtils/Private/OnlineBeaconHost.cpp
new file mode 100644
--- /dev/null
+++ b/Engine/Plugins/Online/OnlineSubsystemUtils/Source/OnlineSubsystemUtils/Private/OnlineBeaconHost.cpp
@@ -0,0 +1,370 @@
+// Copyright 1998-2016 Epic Games, Inc. All Rights Reserved.
+
+#include "OnlineSubsystemUtilsPrivatePCH.h"
+#include "OnlineBeaconHost.h"
+#include "Net/UnrealNetwork.h"
+#include "Net/DataChannel.h"
+#include "OnlineBeaconClient.h"
+#include "OnlineBeaconHostObject.h"
+#include "NetworkVersion.h"
+
+AOnlineBeaconHost::AOnlineBeaconHost(const FObjectInitializer& ObjectInitializer) :
+	Super(ObjectInitializer)
+{
+	NetDriverName = FName(TEXT("BeaconDriverHost"));
+}
+
+void AOnlineBeaconHost::OnNetCleanup(UNetConnection* Connection)
+{
+	UE_LOG(LogBeacon, Error, TEXT("Cleaning up a beacon host!"));
+	ensure(0);
+}
+
+bool AOnlineBeaconHost::InitHost()
+{
+	FURL URL(nullptr, TEXT(""), TRAVEL_Absolute);
+
+	// Allow the command line to override the default port
+	int32 PortOverride;
+	if (FParse::Value(FCommandLine::Get(), TEXT("BeaconPort="), PortOverride) && PortOverride != 0)
+	{
+		ListenPort = PortOverride;
+	}
+
+	URL.Port = ListenPort;
+	if(URL.Valid)
+	{
+		if (InitBase() && NetDriver)
+		{
+			FString Error;
+			if (NetDriver->InitListen(this, URL, false, Error))
+			{
+				ListenPort = URL.Port;
+				NetDriver->SetWorld(GetWorld());
+				NetDriver->Notify = this;
+				NetDriver->InitialConnectTimeout = BeaconConnectionInitialTimeout;
+				NetDriver->ConnectionTimeout = BeaconConnectionTimeout;
+				return true;
+			}
+			else
+			{
+				// error initializing the network stack...
+				UE_LOG(LogBeacon, Log, TEXT("AOnlineBeaconHost::InitHost failed"));
+				OnFailure();
+			}
+		}
+	}
+
+	return false;
+}
+
+void AOnlineBeaconHost::HandleNetworkFailure(UWorld* World, class UNetDriver* InNetDriver, ENetworkFailure::Type FailureType, const FString& ErrorString)
+{
+	if (InNetDriver && InNetDriver->NetDriverName == NetDriverName)
+	{
+		// Timeouts from clients are ignored
+		if (FailureType != ENetworkFailure::ConnectionTimeout)
+		{
+			Super::HandleNetworkFailure(World, InNetDriver, FailureType, ErrorString);
+		}
+	}
+}
+
+void AOnlineBeaconHost::NotifyControlMessage(UNetConnection* Connection, uint8 MessageType, class FInBunch& Bunch)
+{
+	if(NetDriver->ServerConnection == nullptr)
+	{
+		bool bCloseConnection = false;
+
+		// We are the server.
+#if !(UE_BUILD_SHIPPING && WITH_EDITOR)
+		UE_LOG(LogBeacon, Verbose, TEXT("%s[%s] Host received: %s"), *GetName(), Connection ? *Connection->GetName() : TEXT("Invalid"), FNetControlMessageInfo::GetName(MessageType));
+#endif
+		switch (MessageType)
+		{
+		case NMT_Hello:
+			{
+				UE_LOG(LogBeacon, Log, TEXT("Beacon Hello"));
+				uint8 IsLittleEndian;
+
+				uint32 RemoteNetworkVersion = 0;
+				uint32 LocalNetworkVersion = FNetworkVersion::GetLocalNetworkVersion();
+
+				FNetControlMessage<NMT_Hello>::Receive(Bunch, IsLittleEndian, RemoteNetworkVersion);
+
+				if (!FNetworkVersion::IsNetworkCompatible(LocalNetworkVersion, RemoteNetworkVersion))
+				{
+					UE_LOG(LogBeacon, Log, TEXT("Client not network compatible %s"), *Connection->GetName());
+					FNetControlMessage<NMT_Upgrade>::Send(Connection, LocalNetworkVersion);
+					bCloseConnection = true;
+				}
+				else
+				{
+					Connection->Challenge = FString::Printf(TEXT("%08X"), FPlatformTime::Cycles());
+					FNetControlMessage<NMT_BeaconWelcome>::Send(Connection);
+					Connection->FlushNet();
+				}
+				break;
+			}
+		case NMT_Netspeed:
+			{
+				int32 Rate;
+				FNetControlMessage<NMT_Netspeed>::Receive(Bunch, Rate);
+				Connection->CurrentNetSpeed = FMath::Clamp(Rate, 1800, NetDriver->MaxClientRate);
+				UE_LOG(LogBeacon, Log, TEXT("Client netspeed is %i"), Connection->CurrentNetSpeed);
+				break;
+			}
+		case NMT_BeaconJoin:
+			{
+				FString ErrorMsg;
+				FString BeaconType;
+				FUniqueNetIdRepl UniqueId;
+				FNetControlMessage<NMT_BeaconJoin>::Receive(Bunch, BeaconType, UniqueId);
+				UE_LOG(LogBeacon, Log, TEXT("Beacon Join %s %s"), *BeaconType, *UniqueId.ToDebugString());
+
+				if (Connection->ClientWorldPackageName == NAME_None)
+				{
+					AOnlineBeaconClient* ClientActor = GetClientActor(Connection);
+					if (ClientActor == nullptr)
+					{
+						UWorld* World = GetWorld();
+						Connection->ClientWorldPackageName = World->GetOutermost()->GetFName();
+
+						AOnlineBeaconClient* NewClientActor = nullptr;
+						FOnBeaconSpawned* OnBeaconSpawnedDelegate = OnBeaconSpawnedMapping.Find(BeaconType);
+						if (OnBeaconSpawnedDelegate && OnBeaconSpawnedDelegate->IsBound())
+						{
+							NewClientActor = OnBeaconSpawnedDelegate->Execute(Connection);
+						}
+
+						if (NewClientActor && BeaconType == NewClientActor->GetBeaconType())
+						{
+							NewClientActor->SetConnectionState(EBeaconConnectionState::Pending);
+
+							FNetworkGUID NetGUID = Connection->Driver->GuidCache->AssignNewNetGUID_Server(NewClientActor);
+							NewClientActor->SetNetConnection(Connection);
+							Connection->PlayerId = UniqueId;
+							Connection->OwningActor = NewClientActor;
+							NewClientActor->Role = ROLE_Authority;
+							NewClientActor->SetReplicates(false);
+							check(NetDriverName == NetDriver->NetDriverName);
+							NewClientActor->SetNetDriverName(NetDriverName);
+							ClientActors.Add(NewClientActor);
+							FNetControlMessage<NMT_BeaconAssignGUID>::Send(Connection, NetGUID);
+						}
+						else
+						{
+							ErrorMsg = NSLOCTEXT("NetworkErrors", "BeaconSpawnFailureError", "Join failure, Couldn't spawn beacon.").ToString();
+						}
+					}
+					else
+					{
+						ErrorMsg = NSLOCTEXT("NetworkErrors", "BeaconSpawnExistingActorError", "Join failure, existing beacon actor.").ToString();
+					}
+				}
+				else
+				{
+					ErrorMsg = NSLOCTEXT("NetworkErrors", "BeaconSpawnClientWorldPackageNameError", "Join failure, existing ClientWorldPackageName.").ToString();
+				}
+
+				if (!ErrorMsg.IsEmpty())
+				{
+					UE_LOG(LogBeacon, Log, TEXT("%s: %s"), *Connection->GetName(), *ErrorMsg);
+					FNetControlMessage<NMT_Failure>::Send(Connection, ErrorMsg);
+					bCloseConnection = true;
+				}
+
+				break;
+			}
+		case NMT_BeaconNetGUIDAck:
+			{
+				FString ErrorMsg;
+				FString BeaconType;
+				FNetControlMessage<NMT_BeaconNetGUIDAck>::Receive(Bunch, BeaconType);
+
+				AOnlineBeaconClient* ClientActor = GetClientActor(Connection);
+				if (ClientActor && BeaconType == ClientActor->GetBeaconType())
+				{
+					FOnBeaconConnected* OnBeaconConnectedDelegate = OnBeaconConnectedMapping.Find(BeaconType);
+					if (OnBeaconConnectedDelegate)
+					{
+						ClientActor->SetReplicates(true);
+						ClientActor->SetAutonomousProxy(true);
+						ClientActor->SetConnectionState(EBeaconConnectionState::Open);
+						// Send an RPC to the client to open the actor channel and guarantee RPCs will work
+						ClientActor->ClientOnConnected();
+						UE_LOG(LogBeacon, Log, TEXT("Handshake complete for %s!"), *ClientActor->GetName());
+
+						OnBeaconConnectedDelegate->ExecuteIfBound(ClientActor, Connection);
+					}
+					else
+					{
+						// Failed to connect.
+						ErrorMsg = NSLOCTEXT("NetworkErrors", "BeaconSpawnNetGUIDAckError1", "Join failure, no host object at NetGUIDAck.").ToString();
+					}
+				}
+				else
+				{
+					// Failed to connect.
+					ErrorMsg = NSLOCTEXT("NetworkErrors", "BeaconSpawnNetGUIDAckError2", "Join failure, no actor at NetGUIDAck.").ToString();
+				}
+
+				if (!ErrorMsg.IsEmpty())
+				{
+					UE_LOG(LogBeacon, Log, TEXT("%s: %s"), *Connection->GetName(), *ErrorMsg);
+					FNetControlMessage<NMT_Failure>::Send(Connection, ErrorMsg);
+					bCloseConnection = true;
+				}
+
+				break;
+			}
+		case NMT_BeaconWelcome:
+		case NMT_BeaconAssignGUID:
+		default:
+			{
+				FString ErrorMsg = NSLOCTEXT("NetworkErrors", "BeaconSpawnUnexpectedError", "Join failure, unexpected control message.").ToString();
+				UE_LOG(LogBeacon, Log, TEXT("%s: %s"), *Connection->GetName(), *ErrorMsg);
+				FNetControlMessage<NMT_Failure>::Send(Connection, ErrorMsg);
+				bCloseConnection = true;
+			}
+			break;
+		}
+
+		if (bCloseConnection)
+		{		
+			UE_LOG(LogBeacon, Verbose, TEXT("Closing connection %s: %s"), *Connection->GetName(), *Connection->PlayerId.ToDebugString());
+			AOnlineBeaconClient* ClientActor = GetClientActor(Connection);
+			if (ClientActor)
+			{
+				UE_LOG(LogBeacon, Verbose, TEXT("- BeaconActor: %s %s"), *ClientActor->GetName(), *ClientActor->GetBeaconType());
+				AOnlineBeaconHostObject* BeaconHostObject = GetHost(ClientActor->GetBeaconType());
+				if (BeaconHostObject)
+				{
+					UE_LOG(LogBeacon, Verbose, TEXT("- HostObject: %s"), *BeaconHostObject->GetName());
+					BeaconHostObject->NotifyClientDisconnected(ClientActor);
+				}
+
+				RemoveClientActor(ClientActor);
+			}
+
+			Connection->FlushNet(true);
+			Connection->Close();
+			UE_LOG(LogBeacon, Verbose, TEXT("--------------------------------"));
+		}
+	}
+}
+
+void AOnlineBeaconHost::DisconnectClient(AOnlineBeaconClient* ClientActor)
+{
+	if (ClientActor && ClientActor->GetConnectionState() != EBeaconConnectionState::Closed && !ClientActor->IsPendingKill())
+	{
+		ClientActor->SetConnectionState(EBeaconConnectionState::Closed);
+
+		UNetConnection* Connection = ClientActor->GetNetConnection();
+
+		UE_LOG(LogBeacon, Log, TEXT("DisconnectClient for %s. UNetConnection %s UNetDriver %s State %d"), 
+			*GetNameSafe(ClientActor), 
+			*GetNameSafe(Connection), 
+			Connection ? *GetNameSafe(Connection->Driver) : TEXT("null"),
+			Connection ? Connection->State : -1);
+
+		// Closing the connection will start the chain of events leading to the removal from lists and destruction of the actor
+		if (Connection && Connection->State != USOCK_Closed)
+		{
+			Connection->FlushNet(true);
+			Connection->Close();
+		}
+	}
+}
+
+AOnlineBeaconClient* AOnlineBeaconHost::GetClientActor(UNetConnection* Connection)
+{
+	for (int32 ClientIdx=0; ClientIdx < ClientActors.Num(); ClientIdx++)
+	{
+		if (ClientActors[ClientIdx]->GetNetConnection() == Connection)
+		{
+			return ClientActors[ClientIdx];
+		}
+	}
+
+	return nullptr;
+}
+
+void AOnlineBeaconHost::RemoveClientActor(AOnlineBeaconClient* ClientActor)
+{
+	if (ClientActor)
+	{
+		ClientActors.RemoveSingleSwap(ClientActor);
+		if (!ClientActor->IsPendingKillPending())
+		{
+			ClientActor->Destroy();
+		}
+	}
+}
+
+void AOnlineBeaconHost::RegisterHost(AOnlineBeaconHostObject* NewHostObject)
+{
+	const FString& BeaconType = NewHostObject->GetBeaconType();
+	if (GetHost(BeaconType) == NULL)
+	{
+		NewHostObject->SetOwner(this);
+		OnBeaconSpawned(BeaconType).BindUObject(NewHostObject, &AOnlineBeaconHostObject::SpawnBeaconActor);
+		OnBeaconConnected(BeaconType).BindUObject(NewHostObject, &AOnlineBeaconHostObject::OnClientConnected);
+	}
+	else
+	{
+		UE_LOG(LogBeacon, Warning, TEXT("Beacon host type %s already exists"), *BeaconType);
+	}
+}
+
+void AOnlineBeaconHost::UnregisterHost(const FString& BeaconType)
+{
+	AOnlineBeaconHostObject* HostObject = GetHost(BeaconType);
+	if (HostObject)
+	{
+		HostObject->Unregister();
+	}
+	
+	OnBeaconSpawned(BeaconType).Unbind();
+	OnBeaconConnected(BeaconType).Unbind();
+}
+
+AOnlineBeaconHostObject* AOnlineBeaconHost::GetHost(const FString& BeaconType)
+{
+	for (int32 HostIdx=0; HostIdx < Children.Num(); HostIdx++)
+	{
+		AOnlineBeaconHostObject* HostObject = Cast<AOnlineBeaconHostObject>(Children[HostIdx]);
+		if (HostObject && HostObject->GetBeaconType() == BeaconType)
+		{
+			return HostObject;
+		}
+	}
+
+	return nullptr;
+}
+
+AOnlineBeaconHost::FOnBeaconSpawned& AOnlineBeaconHost::OnBeaconSpawned(const FString& BeaconType)
+{ 
+	FOnBeaconSpawned* BeaconDelegate = OnBeaconSpawnedMapping.Find(BeaconType);
+	if (BeaconDelegate == nullptr)
+	{
+		FOnBeaconSpawned NewDelegate;
+		OnBeaconSpawnedMapping.Add(BeaconType, NewDelegate);
+		BeaconDelegate = OnBeaconSpawnedMapping.Find(BeaconType);
+	}
+
+	return *BeaconDelegate; 
+}
+
+AOnlineBeaconHost::FOnBeaconConnected& AOnlineBeaconHost::OnBeaconConnected(const FString& BeaconType)
+{ 
+	FOnBeaconConnected* BeaconDelegate = OnBeaconConnectedMapping.Find(BeaconType);
+	if (BeaconDelegate == nullptr)
+	{
+		FOnBeaconConnected NewDelegate;
+		OnBeaconConnectedMapping.Add(BeaconType, NewDelegate);
+		BeaconDelegate = OnBeaconConnectedMapping.Find(BeaconType);
+	}
+
+	return *BeaconDelegate; 
+}
diff --git a/Engine/Plugins/Online/OnlineSubsystemUtils/Source/OnlineSubsystemUtils/Private/OnlineEngineInterfaceImpl.h b/Engine/Plugins/Online/OnlineSubsystemUtils/Source/OnlineSubsystemUtils/Private/OnlineEngineInterfaceImpl.h
new file mode 100644
--- /dev/null
+++ b/Engine/Plugins/Online/OnlineSubsystemUtils/Source/OnlineSubsystemUtils/Private/OnlineEngineInterfaceImpl.h
@@ -0,0 +1,138 @@
+// Copyright 1998-2016 Epic Games, Inc. All Rights Reserved.
+
+#pragma once
+
+#include "Net/OnlineEngineInterface.h"
+#include "OnlineEngineInterfaceImpl.generated.h"
+
+class UWorld;
+class FVoicePacket;
+class FArchive;
+class FUniqueNetId;
+
+UCLASS(config = Engine)
+class ONLINESUBSYSTEMUTILS_API UOnlineEngineInterfaceImpl : public UOnlineEngineInterface
+{
+	GENERATED_UCLASS_BODY()
+
+	/**
+	 * Subsystem
+	 */
+public:
+
+	virtual bool IsLoaded(FName OnlineIdentifier) override;
+	virtual FName GetOnlineIdentifier(FWorldContext& WorldContext) override;
+	virtual bool DoesInstanceExist(FName OnlineIdentifier) override;
+	virtual void ShutdownOnlineSubsystem(FName OnlineIdentifier) override;
+	virtual void DestroyOnlineSubsystem(FName OnlineIdentifier) override;
+
+private:
+
+	FName GetOnlineIdentifier(UWorld* World);
+
+	/**
+	 * Identity
+	 */
+public:
+
+	virtual TSharedPtr<const FUniqueNetId> CreateUniquePlayerId(const FString& Str) override;
+	virtual TSharedPtr<const FUniqueNetId> GetUniquePlayerId(UWorld* World, int32 LocalUserNum) override;
+
+	virtual FString GetPlayerNickname(UWorld* World, const FUniqueNetId& UniqueId) override;
+	virtual bool GetPlayerPlatformNickname(UWorld* World, int32 LocalUserNum, FString& OutNickname) override;
+
+	virtual bool AutoLogin(UWorld* World, int32 LocalUserNum, const FOnlineAutoLoginComplete& InCompletionDelegate) override;
+	virtual bool IsLoggedIn(UWorld* World, int32 LocalUserNum) override;
+
+private:
+
+	FDelegateHandle OnLoginCompleteDelegateHandle;
+	void OnAutoLoginComplete(int32 LocalUserNum, bool bWasSuccessful, const FUniqueNetId& UserId, const FString& Error, FName OnlineIdentifier, FOnlineAutoLoginComplete InCompletionDelegate);
+
+	/**
+	 * Session
+	 */
+public:
+
+	virtual void StartSession(UWorld* World, FName SessionName, FOnlineSessionStartComplete& InCompletionDelegate) override;
+	virtual void EndSession(UWorld* World, FName SessionName, FOnlineSessionEndComplete& InCompletionDelegate) override;
+	virtual bool DoesSessionExist(UWorld* World, FName SessionName) override;
+
+	virtual bool GetSessionJoinability(UWorld* World, FName SessionName, FJoinabilitySettings& OutSettings) override;
+	virtual void UpdateSessionJoinability(UWorld* World, FName SessionName, bool bPublicSearchable, bool bAllowInvites, bool bJoinViaPresence, bool bJoinViaPresenceFriendsOnly) override;
+
+	virtual void RegisterPlayer(UWorld* World, FName SessionName, const FUniqueNetId& UniqueId, bool bWasInvited) override;
+	virtual void UnregisterPlayer(UWorld* World, FName SessionName, const FUniqueNetId& UniqueId) override;
+
+	virtual bool GetResolvedConnectString(UWorld* World, FName SessionName, FString& URL) override;
+
+private:
+
+	/** Mapping of delegate handles for each online StartSession() call while in flight */
+	TMap<FName, FDelegateHandle> OnStartSessionCompleteDelegateHandles;
+	void OnStartSessionComplete(FName SessionName, bool bWasSuccessful, FName OnlineIdentifier, FOnlineSessionStartComplete CompletionDelegate);
+
+	/** Mapping of delegate handles for each online EndSession() call while in flight */
+	TMap<FName, FDelegateHandle> OnEndSessionCompleteDelegateHandles;
+	void OnEndSessionComplete(FName SessionName, bool bWasSuccessful, FName OnlineIdentifier, FOnlineSessionEndComplete CompletionDelegate);
+
+	/**
+	 * Voice
+	 */
+public:
+
+	virtual TSharedPtr<FVoicePacket> GetLocalPacket(UWorld* World, uint8 LocalUserNum) override;
+	virtual TSharedPtr<FVoicePacket> SerializeRemotePacket(UWorld* World, FArchive& Ar) override;
+
+	virtual void StartNetworkedVoice(UWorld* World, uint8 LocalUserNum) override;
+	virtual void StopNetworkedVoice(UWorld* World, uint8 LocalUserNum) override;
+	virtual void ClearVoicePackets(UWorld* World) override;
+
+	virtual bool MuteRemoteTalker(UWorld* World, uint8 LocalUserNum, const FUniqueNetId& PlayerId, bool bIsSystemWide) override;
+	virtual bool UnmuteRemoteTalker(UWorld* World, uint8 LocalUserNum, const FUniqueNetId& PlayerId, bool bIsSystemWide) override;
+
+	virtual int32 GetNumLocalTalkers(UWorld* World) override;
+
+	/**
+	 * External UI
+	 */
+public:
+
+	virtual void ShowLeaderboardUI(UWorld* World, const FString& CategoryName) override;
+	virtual void ShowAchievementsUI(UWorld* World, int32 LocalUserNum) override;
+	virtual void BindToExternalUIOpening(const FOnlineExternalUIChanged& Delegate) override;
+
+private:
+
+	void OnExternalUIChange(bool bInIsOpening, FOnlineExternalUIChanged Delegate);
+
+	/**
+	 * Debug
+	 */
+public:
+
+	virtual void DumpSessionState(UWorld* World) override;
+	virtual void DumpPartyState(UWorld* World) override;
+	virtual void DumpVoiceState(UWorld* World) override;
+	virtual void DumpChatState(UWorld* World) override;
+
+#if WITH_EDITOR
+	/**
+	 * PIE Utilities
+	 */
+public:
+
+	virtual bool SupportsOnlinePIE() override;
+	virtual int32 GetNumPIELogins() override;
+	virtual void SetForceDedicated(FName OnlineIdentifier, bool bForce) override;
+	virtual void LoginPIEInstance(FName OnlineIdentifier, int32 LocalUserNum, int32 PIELoginNum, FOnPIELoginComplete& CompletionDelegate) override;
+#endif
+
+private:
+
+	/** Mapping of delegate handles for each online Login() call while in flight */
+	TMap<FName, FDelegateHandle> OnLoginPIECompleteDelegateHandlesForPIEInstances;
+	void OnPIELoginComplete(int32 LocalUserNum, bool bWasSuccessful, const FUniqueNetId& UserId, const FString& Error, FName OnlineIdentifier, FOnlineAutoLoginComplete InCompletionDelegate);
+
+};
+
