Copying //Tasks/UE4/Dev-Online-Plugins-v3 to //UE4/Dev-Main (Source: //Tasks/UE4/Dev-Online-Plugins-v3 @ 3047294)

diff --git a/Engine/Plugins/Online/OnlineSubsystemUtils/Source/OnlineSubsystemUtils/Private/OnlinePIESettings.cpp b/Engine/Plugins/Online/OnlineSubsystemUtils/Source/OnlineSubsystemUtils/Private/OnlinePIESettings.cpp
new file mode 100644
--- /dev/null
+++ b/Engine/Plugins/Online/OnlineSubsystemUtils/Source/OnlineSubsystemUtils/Private/OnlinePIESettings.cpp
@@ -0,0 +1,168 @@
+// Copyright 1998-2016 Epic Games, Inc. All Rights Reserved.
+
+#include "OnlineSubsystemUtilsPrivatePCH.h"
+#include "OnlinePIESettings.h"
+#include "AES.h"
+
+const int32 ONLINEPIE_XOR_KEY = 0xdeadbeef;
+void FPIELoginSettingsInternal::Encrypt()
+{
+	if (Token.Len() > 0)
+	{
+		TArray<TCHAR> SrcCharArray = Token.GetCharArray();
+		int32 SrcSize = SrcCharArray.Num() * sizeof(TCHAR);
+		const int64 PaddedEncryptedFileSize = Align(SrcSize + 1, FAES::AESBlockSize);
+
+		TokenBytes.Empty(PaddedEncryptedFileSize);
+		TokenBytes.AddUninitialized(PaddedEncryptedFileSize);
+
+		// Store the length of the password
+		ensure(SrcSize < MAX_uint8);
+		TokenBytes[0] = static_cast<uint8>(SrcSize);
+
+		// Copy the password in, leaving the random uninitialized data at the end
+		FMemory::Memcpy(TokenBytes.GetData() + 1, SrcCharArray.GetData(), SrcSize);
+
+		// XOR Cipher
+		int32 NumXors = PaddedEncryptedFileSize / sizeof(int32);
+		int32* Password = (int32*)(TokenBytes.GetData());
+		for (int32 i = 0; i < NumXors; i++)
+		{
+			Password[i] = Password[i] ^ ONLINEPIE_XOR_KEY;
+		}
+
+		//FAES::EncryptData(TokenBytes.GetData(), PaddedEncryptedFileSize);
+	}
+	else
+	{
+		TokenBytes.Empty();
+	}
+}
+
+void FPIELoginSettingsInternal::Decrypt()
+{
+	if (TokenBytes.Num() > 0)
+	{
+		const int64 PaddedEncryptedFileSize = Align(TokenBytes.Num(), FAES::AESBlockSize);
+		if (PaddedEncryptedFileSize > 0 && PaddedEncryptedFileSize == TokenBytes.Num())
+		{
+			TArray<uint8> TempArray;
+			TempArray.Empty(PaddedEncryptedFileSize);
+			TempArray.AddUninitialized(PaddedEncryptedFileSize);
+			FMemory::Memcpy(TempArray.GetData(), TokenBytes.GetData(), PaddedEncryptedFileSize);
+			
+			// XOR Cipher
+			int32 NumXors = PaddedEncryptedFileSize / sizeof(int32);
+			int32* TempArrayPtr = (int32*)(TempArray.GetData());
+			for (int32 i = 0; i < NumXors; i++)
+			{
+				TempArrayPtr[i] = TempArrayPtr[i] ^ ONLINEPIE_XOR_KEY;
+			}
+
+			//FAES::DecryptData(TempArray.GetData(), PaddedEncryptedFileSize);
+
+			// Validate the unencrypted data (stored size less than total data, null terminated character where it should be)
+			int32 PasswordDataSize = TempArray[0];
+			int32 PasswordLength = PasswordDataSize / sizeof(TCHAR);
+			TCHAR* Password = (TCHAR*)(TempArray.GetData() + 1);
+			if (PasswordDataSize < TempArray.Num() &&
+				Password[PasswordLength - 1] == '\0')
+			{
+				Token = FString(Password);
+			}
+			else
+			{
+				Token.Empty();
+				TokenBytes.Empty();
+			}
+		}
+	}
+	else
+	{
+		Token.Empty();
+	}
+}
+
+UOnlinePIESettings::UOnlinePIESettings(const FObjectInitializer& ObjectInitializer)
+	: Super(ObjectInitializer)
+{
+	bOnlinePIEEnabled = false;
+	CategoryName = FName(TEXT("LevelEditor"));
+}
+
+void UOnlinePIESettings::PostInitProperties()
+{
+	Super::PostInitProperties();
+
+	for (FPIELoginSettingsInternal& Login : Logins)
+	{
+		Login.Decrypt();
+	}
+}
+
+#if WITH_EDITOR
+void UOnlinePIESettings::PostEditChangeProperty(FPropertyChangedEvent& PropertyChangedEvent)
+{
+	Super::PostEditChangeProperty(PropertyChangedEvent);
+
+	if (PropertyChangedEvent.Property != NULL)
+	{
+		FName MemberPropName = PropertyChangedEvent.MemberProperty->GetFName();
+		FName SubPropName = PropertyChangedEvent.Property->GetFName();
+		if (MemberPropName == "bOnlinePIEEnabled")
+		{
+			// Possibly get rid of the null subsystem in favor of the real default or if we are disabling online pie then get rid of the real subsystem to replace it with null
+			IOnlineSubsystem::ReloadDefaultSubsystem();
+		}
+		if (MemberPropName == "Logins")
+		{
+			if (SubPropName == "Id")
+			{
+				for (FPIELoginSettingsInternal& Login : Logins)
+				{
+					// Remove any whitespace from login input
+					Login.Id = Login.Id.Trim().TrimTrailing();
+				}
+			}
+			else if (SubPropName == "Token")
+			{
+				for (FPIELoginSettingsInternal& Login : Logins)
+				{
+					// Remove any whitespace from login input
+					Login.Token = Login.Token.Trim().TrimTrailing();
+					// Encrypt the password
+					Login.Encrypt();
+				}
+			}
+			else if (SubPropName == "Type")
+			{
+				for (FPIELoginSettingsInternal& Login : Logins)
+				{
+					// Remove any whitespace from login input
+					Login.Type = Login.Type.Trim().TrimTrailing();
+				}
+			}
+
+			if (PropertyChangedEvent.ChangeType != EPropertyChangeType::ArrayAdd)
+			{
+				bool bOneLoginValid = false;
+				for (FPIELoginSettingsInternal& Login : Logins)
+				{
+					if (Login.IsValid())
+					{
+						bOneLoginValid = true;
+						break;
+					}
+				}
+
+				if (!bOneLoginValid)
+				{
+					// Disable PIE logins when there are no logins available
+					bOnlinePIEEnabled = false;
+				}
+			}
+		}
+	}
+}
+
+#endif // WITH_EDITOR
diff --git a/Engine/Plugins/Online/OnlineSubsystemUtils/Source/OnlineSubsystemUtils/Private/OnlinePIESettings.h b/Engine/Plugins/Online/OnlineSubsystemUtils/Source/OnlineSubsystemUtils/Private/OnlinePIESettings.h
new file mode 100644
--- /dev/null
+++ b/Engine/Plugins/Online/OnlineSubsystemUtils/Source/OnlineSubsystemUtils/Private/OnlinePIESettings.h
@@ -0,0 +1,67 @@
+// Copyright 1998-2016 Epic Games, Inc. All Rights Reserved.
+#pragma once
+
+#include "OnlinePIESettings.generated.h"
+
+/**
+ * Stores PIE login credentials
+ */
+USTRUCT()
+struct FPIELoginSettingsInternal
+{
+public:
+
+	GENERATED_USTRUCT_BODY()
+
+	/** Id of the user logging in (email, display name, facebook id, etc) */
+	UPROPERTY(EditAnywhere, Category = "Logins", meta = (DisplayName = "User Id", Tooltip = "Id of the user logging in (email, display name, facebook id, etc)"))
+	FString Id;
+	/** Credentials of the user logging in (password or auth token) */
+	UPROPERTY(EditAnywhere, Transient, Category = "Logins", meta = (DisplayName = "Password", Tooltip = "Credentials of the user logging in (password or auth token)", PasswordField = true))
+	FString Token;
+	/** Type of account. Needed to identity the auth method to use (epic, internal, facebook, etc) */
+	UPROPERTY(EditAnywhere, Category = "Logins", meta = (DisplayName = "Type", Tooltip = "Type of account. Needed to identity the auth method to use (epic, internal, facebook, etc)"))
+	FString Type;
+	/** Token stored as an array of bytes, encrypted */
+	UPROPERTY()
+	TArray<uint8> TokenBytes;
+
+	/** @return true if the credentials are valid, false otherwise */
+	bool IsValid() const { return !Id.IsEmpty() && !Token.IsEmpty() && !Type.IsEmpty(); }
+
+	/**
+	 * Encrypt the Token field into the TokenBytes field
+	 */
+	void Encrypt();
+	/**
+	 * Decrypt the TokenBytes field into the Token field
+	 */
+	void Decrypt();
+};
+
+/**
+ * Setup up login credentials for the Play In Editor (PIE) feature
+ */
+UCLASS(config=EditorPerProjectUserSettings, meta = (DisplayName = "Play Credentials"))
+class UOnlinePIESettings : public UDeveloperSettings
+{
+	GENERATED_BODY()
+
+public:
+	UOnlinePIESettings(const FObjectInitializer& ObjectInitializer);
+
+	/** Will Play In Editor (PIE) attempt to login to a platform service before launching the instance */
+	UPROPERTY(config, EditAnywhere, Category = "Logins", meta = (DisplayName = "Enable Logins", Tooltip = "Attempt to login with user credentials on a backend service before launching the PIE instance."))
+	bool bOnlinePIEEnabled;
+
+	/** Array of credentials to use, one for each Play In Editor (PIE) instance */
+	UPROPERTY(config, EditAnywhere, Category = "Logins", meta = (DisplayName = "Credentials", Tooltip = "Login credentials, at least one for each instance of PIE that is intended to be run"))
+	TArray<FPIELoginSettingsInternal> Logins;
+
+	// Begin UObject Interface
+	virtual void PostInitProperties() override;
+#if WITH_EDITOR
+	virtual void PostEditChangeProperty(FPropertyChangedEvent& PropertyChangedEvent);
+#endif // WITH_EDITOR
+	// End UObject Interface
+};
diff --git a/Engine/Plugins/Online/OnlineSubsystemUtils/Source/OnlineSubsystemUtils/Private/OnlineSubsystemUtilsModule.cpp b/Engine/Plugins/Online/OnlineSubsystemUtils/Source/OnlineSubsystemUtils/Private/OnlineSubsystemUtilsModule.cpp
new file mode 100644
--- /dev/null
+++ b/Engine/Plugins/Online/OnlineSubsystemUtils/Source/OnlineSubsystemUtils/Private/OnlineSubsystemUtilsModule.cpp
@@ -0,0 +1,108 @@
+// Copyright 1998-2016 Epic Games, Inc. All Rights Reserved.
+
+#include "OnlineSubsystemUtilsPrivatePCH.h"
+#include "OnlineSubsystemUtilsModule.h"
+#include "ModuleManager.h"
+#include "OnlinePIESettings.h"
+
+IMPLEMENT_MODULE(FOnlineSubsystemUtilsModule, OnlineSubsystemUtils);
+
+/**
+ * Concrete implementation of IOnlineSubsystemUtils interface 
+ */
+class FOnlineSubsystemUtils : public IOnlineSubsystemUtils
+{
+public:
+
+	FOnlineSubsystemUtils() {}
+	virtual ~FOnlineSubsystemUtils() {}
+
+	FName GetOnlineIdentifier(const FWorldContext& WorldContext) override
+	{
+#if WITH_EDITOR
+		return FName(*FString::Printf(TEXT(":%s"), *WorldContext.ContextHandle.ToString()));
+#else
+		return NAME_None;
+#endif
+	}
+
+	FName GetOnlineIdentifier(UWorld* World, const FName Subsystem = NAME_None) override
+	{
+#if WITH_EDITOR
+		if (const FWorldContext* WorldContext = GEngine->GetWorldContextFromWorld(World))
+		{
+			return FName(
+				*FString::Printf(TEXT("%s:%s"), !Subsystem.IsNone() ? *Subsystem.ToString() : TEXT(""), *WorldContext->ContextHandle.ToString()));
+		}
+
+		return NAME_None;
+#else
+		return Subsystem;
+#endif
+	}
+
+#if WITH_EDITOR
+	virtual bool SupportsOnlinePIE() const override
+	{
+		const UOnlinePIESettings* OnlinePIESettings = GetDefault<UOnlinePIESettings>();
+		if (OnlinePIESettings->bOnlinePIEEnabled && GetNumPIELogins() > 0)
+		{
+			// If we can't get the identity interface then things are either not configured right or disabled
+			IOnlineIdentityPtr IdentityInt = Online::GetIdentityInterface();
+			return IdentityInt.IsValid();
+		}
+
+		return false;
+	}
+
+	virtual bool IsOnlinePIEEnabled() const override
+	{
+		const UOnlinePIESettings* OnlinePIESettings = GetDefault<UOnlinePIESettings>();
+		return OnlinePIESettings->bOnlinePIEEnabled;
+	}
+
+	virtual int32 GetNumPIELogins() const override
+	{
+		int32 NumValidLogins = 0;
+		const UOnlinePIESettings* OnlinePIESettings = GetDefault<UOnlinePIESettings>();
+		for (const FPIELoginSettingsInternal& Login : OnlinePIESettings->Logins)
+		{
+			if (Login.IsValid())
+			{
+				NumValidLogins++;
+			}
+		}
+	
+		return NumValidLogins;
+	}
+
+	virtual void GetPIELogins(TArray<FOnlineAccountCredentials>& Logins) override
+	{
+		const UOnlinePIESettings* OnlinePIESettings = GetDefault<UOnlinePIESettings>();
+		if (OnlinePIESettings->Logins.Num() > 0)
+		{
+			FOnlineAccountCredentials TempLogin;
+
+			Logins.Empty(OnlinePIESettings->Logins.Num());
+			for (const FPIELoginSettingsInternal& Login : OnlinePIESettings->Logins)
+			{
+				if (Login.IsValid())
+				{
+					new (Logins)FOnlineAccountCredentials(Login.Type, Login.Id, Login.Token);
+				}
+			}
+		}
+	}
+#endif // WITH_EDITOR
+};
+
+void FOnlineSubsystemUtilsModule::StartupModule()
+{
+	SubsystemUtils = new FOnlineSubsystemUtils();
+}
+
+void FOnlineSubsystemUtilsModule::ShutdownModule()
+{
+	delete SubsystemUtils;
+	SubsystemUtils = nullptr;
+}
diff --git a/Engine/Plugins/Online/OnlineSubsystemUtils/Source/OnlineSubsystemUtils/Private/OnlineSubsystemUtilsPrivatePCH.h b/Engine/Plugins/Online/OnlineSubsystemUtils/Source/OnlineSubsystemUtils/Private/OnlineSubsystemUtilsPrivatePCH.h
new file mode 100644
--- /dev/null
+++ b/Engine/Plugins/Online/OnlineSubsystemUtils/Source/OnlineSubsystemUtils/Private/OnlineSubsystemUtilsPrivatePCH.h
@@ -0,0 +1,41 @@
+// Copyright 1998-2016 Epic Games, Inc. All Rights Reserved.
+
+#pragma once
+
+#include "Core.h"
+#include "Engine.h"
+#include "Sockets.h"
+#include "OnlineSubsystem.h"
+#include "OnlineSubsystemImpl.h"
+#include "OnlineSubsystemUtils.h"
+#include "OnlineSubsystemUtilsModule.h"
+#include "SocketSubsystem.h"
+#include "ModuleManager.h"
+#include "Classes/IpConnection.h"
+#include "Classes/IpNetDriver.h"
+
+class FUniqueNetId;
+
+// Helper class for various methods to reduce the call hierarchy
+struct FOnlineSubsystemBPCallHelper
+{
+public:
+	FOnlineSubsystemBPCallHelper(const TCHAR* CallFunctionContext, UWorld* World, FName SystemName = NAME_None);
+
+	void QueryIDFromPlayerController(APlayerController* PlayerController);
+
+	bool IsValid() const
+	{
+		return UserID.IsValid() && (OnlineSub != nullptr);
+	}
+
+public:
+	TSharedPtr<const FUniqueNetId> UserID;
+	IOnlineSubsystem* const OnlineSub;
+	const TCHAR* FunctionContext;
+};
+
+#define INVALID_INDEX -1
+
+
+
diff --git a/Engine/Plugins/Online/OnlineSubsystemUtils/Source/OnlineSubsystemUtils/Private/TestBeaconClient.cpp b/Engine/Plugins/Online/OnlineSubsystemUtils/Source/OnlineSubsystemUtils/Private/TestBeaconClient.cpp
new file mode 100644
--- /dev/null
+++ b/Engine/Plugins/Online/OnlineSubsystemUtils/Source/OnlineSubsystemUtils/Private/TestBeaconClient.cpp
@@ -0,0 +1,42 @@
+// Copyright 1998-2016 Epic Games, Inc. All Rights Reserved.
+
+#include "OnlineSubsystemUtilsPrivatePCH.h"
+#include "TestBeaconClient.h"
+
+ATestBeaconClient::ATestBeaconClient(const FObjectInitializer& ObjectInitializer) :
+	Super(ObjectInitializer)
+{
+}
+
+void ATestBeaconClient::OnFailure()
+{
+#if !UE_BUILD_SHIPPING
+	UE_LOG(LogBeacon, Verbose, TEXT("Test beacon connection failure, handling connection timeout."));
+#endif
+	Super::OnFailure();
+}
+
+void ATestBeaconClient::ClientPing_Implementation()
+{
+#if !UE_BUILD_SHIPPING
+	UE_LOG(LogBeacon, Log, TEXT("Ping"));
+	ServerPong();
+#endif
+}
+
+bool ATestBeaconClient::ServerPong_Validate()
+{
+#if !UE_BUILD_SHIPPING
+	return true;
+#else
+	return false;
+#endif
+}
+
+void ATestBeaconClient::ServerPong_Implementation()
+{
+#if !UE_BUILD_SHIPPING
+	UE_LOG(LogBeacon, Log, TEXT("Pong"));
+	ClientPing();
+#endif
+}
diff --git a/Engine/Plugins/Online/OnlineSubsystemUtils/Source/OnlineSubsystemUtils/Private/Tests/TestAchievementsInterface.cpp b/Engine/Plugins/Online/OnlineSubsystemUtils/Source/OnlineSubsystemUtils/Private/Tests/TestAchievementsInterface.cpp
new file mode 100644
--- /dev/null
+++ b/Engine/Plugins/Online/OnlineSubsystemUtils/Source/OnlineSubsystemUtils/Private/Tests/TestAchievementsInterface.cpp
@@ -0,0 +1,207 @@
+// Copyright 1998-2016 Epic Games, Inc. All Rights Reserved.
+
+#include "OnlineSubsystemUtilsPrivatePCH.h"
+#include "ModuleManager.h"
+#include "TestAchievementsInterface.h"
+
+#if WITH_DEV_AUTOMATION_TESTS
+
+/**
+ * The process of this test is:
+ *
+ * [Read Achievements] -> [ReadAchievementsDelegate]
+ * -> [Read Achievement Descriptions] -> [Read Achievement Descriptions Delegate]
+ * -> -> [Write Achievements] -> [Achievement Unlocked Delegate]
+ * -> -> -> [delete self]
+ */
+
+
+void FTestAchievementsInterface::Test(UWorld* InWorld)
+{
+	UE_LOG(LogOnline, Display, TEXT("FTestAchievementsInterface::Test"));
+
+	IOnlineSubsystem* OnlineSub = Online::GetSubsystem(InWorld, FName(*SubsystemName));
+	check(OnlineSub); 
+
+	if (OnlineSub->GetIdentityInterface().IsValid())
+	{
+		UserId = OnlineSub->GetIdentityInterface()->GetUniquePlayerId(0);
+
+		OnlineAchievements = OnlineSub->GetAchievementsInterface();
+	}
+
+	if( OnlineAchievements.IsValid() && UserId.IsValid() )
+	{
+		ReadAchievements();
+	}
+	else
+	{
+		UE_LOG(LogOnline, Warning, TEXT("TEST FAILED: OSS [%s] does not have a valid achievement interface or identity interface for this test to run."), *SubsystemName );
+		delete this;
+	}
+}
+
+
+void FTestAchievementsInterface::ReadAchievements()
+{
+	// Add our delegate and read achievements
+	OnlineAchievements->QueryAchievements( *UserId.Get(), FOnQueryAchievementsCompleteDelegate::CreateRaw( this, &FTestAchievementsInterface::OnQueryAchievementsComplete ) );
+}
+
+
+void FTestAchievementsInterface::OnQueryAchievementsComplete(const FUniqueNetId& PlayerId, const bool bWasSuccessful)
+{
+	UE_LOG(LogOnline, Display, TEXT("FTestAchievementsInterface::OnQueryAchievementsComplete"));
+
+	if (bWasSuccessful)
+	{
+		UE_LOG(LogOnline, Display, TEXT("Loaded Achievements"));
+
+		TArray<FOnlineAchievement> PlayerAchievements;
+		if (OnlineAchievements->GetCachedAchievements(*UserId.Get(), PlayerAchievements) != EOnlineCachedResult::Success || PlayerAchievements.Num() == 0)
+		{
+			UE_LOG(LogOnline, Warning, TEXT("TEST FAILED: Either GetCachedAchievements() failed or number of achievements is 0"));
+			delete this;
+			return;
+		}
+
+		UE_LOG(LogOnline, Display, TEXT("Number of Achievements: %d"), PlayerAchievements.Num());
+		for(int32 Idx = 0; Idx < PlayerAchievements.Num(); ++Idx)
+		{
+			UE_LOG(LogOnline, Display, TEXT(" Achievement %d: %s"), Idx, *PlayerAchievements[ Idx ].ToDebugString());
+		}
+
+		QueryAchievementDescriptions( PlayerId );
+	}
+	else
+	{
+		// if our test failed, be sure to clean itself up
+		UE_LOG(LogOnline, Display, TEXT("TEST FAILED: Failed to Load Achievements"));
+		delete this;
+	}
+}
+
+
+void FTestAchievementsInterface::QueryAchievementDescriptions(const FUniqueNetId& PlayerId)
+{
+	// Add our delegate and read achievement descriptions
+	OnlineAchievements->QueryAchievementDescriptions( PlayerId, FOnQueryAchievementsCompleteDelegate::CreateRaw( this, &FTestAchievementsInterface::OnQueryAchievementDescriptionsComplete ) );
+}
+
+
+void FTestAchievementsInterface::OnQueryAchievementDescriptionsComplete(const FUniqueNetId& PlayerId, const bool bWasSuccessful)
+{
+	UE_LOG(LogOnline, Display, TEXT("FTestAchievementsInterface::OnQueryAchievementDescriptionsComplete"));
+
+	if (bWasSuccessful)
+	{
+		UE_LOG(LogOnline, Display, TEXT("Loaded Achievement descriptions"));
+
+		TArray<FOnlineAchievement> PlayerAchievements;
+		if (OnlineAchievements->GetCachedAchievements(*UserId.Get(), PlayerAchievements) != EOnlineCachedResult::Success || PlayerAchievements.Num() == 0)
+		{
+			UE_LOG(LogOnline, Warning, TEXT("TEST FAILED: Either GetCachedAchievements() failed or number of achievements is 0"));
+			delete this;
+			return;
+		}
+
+		for(int32 Idx = 0; Idx < PlayerAchievements.Num(); ++Idx)
+		{
+			FOnlineAchievementDesc Desc;
+			if (OnlineAchievements->GetCachedAchievementDescription(PlayerAchievements[ Idx ].Id, Desc) != EOnlineCachedResult::Success)
+			{
+				UE_LOG(LogOnline, Warning, TEXT("Failed to GetCachedAchievementDescription() for achievement '%s'"),
+					*PlayerAchievements[ Idx ].Id
+					);
+				delete this;
+				return;
+			}
+			UE_LOG(LogOnline, Display, TEXT(" Descriptor for achievement '%s': %s"), *PlayerAchievements[ Idx ].Id, *Desc.ToDebugString());
+		}
+		
+		WriteAchievements();
+	}
+	else
+	{
+		// if our test failed, be sure to clean itself up
+		UE_LOG(LogOnline, Warning, TEXT("TEST FAILED: Failed to Load Achievement descriptions"));
+		delete this;
+	}
+}
+
+
+void FTestAchievementsInterface::WriteAchievements()
+{
+	WriteObject = MakeShareable(new FOnlineAchievementsWrite());
+	FOnlineAchievementsWriteRef WriteObjectRef = WriteObject.ToSharedRef();
+
+	TArray<FOnlineAchievement> PlayerAchievements;
+	if (OnlineAchievements->GetCachedAchievements(*UserId.Get(), PlayerAchievements) != EOnlineCachedResult::Success || PlayerAchievements.Num() == 0)
+	{
+		UE_LOG(LogOnline, Warning, TEXT("TEST FAILED: Either GetCachedAchievements() failed or number of achievements is 0"));
+		delete this;
+		return;
+	}
+
+	if (PlayerAchievements.Num() >= 3)
+	{
+		// report progress on first, last and median achievements
+		WriteObject->SetFloatStat(*PlayerAchievements[ 0 ].Id, 1.0f);
+		WriteObject->SetFloatStat(*PlayerAchievements[ PlayerAchievements.Num() - 1 ].Id, 50.0f);
+		WriteObject->SetFloatStat(*PlayerAchievements[ PlayerAchievements.Num() / 2 ].Id, 100.0f);
+	}
+	else
+	{
+		// report progress on all achievements
+		for (int32 IdxAch = 0; IdxAch < PlayerAchievements.Num(); ++IdxAch)
+		{
+			WriteObject->SetFloatStat(*PlayerAchievements[ IdxAch ].Id, 50.0f);
+		}
+	}
+
+	// Add our delegate and write some achievements to the server
+	OnAchievementUnlockedDelegateHandle = OnlineAchievements->AddOnAchievementUnlockedDelegate_Handle( OnAchievementUnlockedDelegate );
+	OnlineAchievements->WriteAchievements( *UserId.Get(), WriteObjectRef, OnAchievementsWrittenDelegate );
+}
+
+
+void FTestAchievementsInterface::OnAchievementsWritten(const FUniqueNetId& PlayerId, bool bWasSuccessful)
+{
+	UE_LOG(LogOnline, Display, TEXT("FTestAchievementsInterface::OnAchievementsWritten( bWasSuccessful = %s )"), bWasSuccessful ? TEXT("true") : TEXT("false"));
+
+	// Clear our delegate and delete this test.
+	OnlineAchievements->ClearOnAchievementUnlockedDelegate_Handle(OnAchievementUnlockedDelegateHandle);
+
+	if (bWasSuccessful && WriteObject->WriteState == EOnlineAsyncTaskState::Done)
+	{
+#if !UE_BUILD_SHIPPING
+		OnlineAchievements->ResetAchievements(*UserId.Get());
+#endif // !UE_BUILD_SHIPPING
+		// TODO: spin for a while?
+		UE_LOG(LogOnline, Display, TEXT("TEST COMPLETED SUCCESSFULLY."));
+	}
+	else
+	{
+		if (WriteObject->WriteState != EOnlineAsyncTaskState::Done)
+		{
+			UE_LOG(LogOnline, Warning, TEXT("TEST FAILED: WriteObject->WriteState is not in state %s, but instead %s"),
+				EOnlineAsyncTaskState::ToString(EOnlineAsyncTaskState::Done),
+				EOnlineAsyncTaskState::ToString(WriteObject->WriteState)
+				);
+		}
+
+		if (!bWasSuccessful)
+		{
+			UE_LOG(LogOnline, Warning, TEXT("TEST FAILED: Write did not complete successfully") );
+		}
+	}
+
+	delete this;
+}
+
+void FTestAchievementsInterface::OnAchievementsUnlocked(const FUniqueNetId& PlayerId, const FString& AchievementId)
+{
+	UE_LOG(LogOnline, Display, TEXT("Achievement Unlocked - %s"), *AchievementId);
+}
+
+#endif //WITH_DEV_AUTOMATION_TESTS
\ No newline at end of file
diff --git a/Engine/Plugins/Online/OnlineSubsystemUtils/Source/OnlineSubsystemUtils/Private/Tests/TestCloudInterface.cpp b/Engine/Plugins/Online/OnlineSubsystemUtils/Source/OnlineSubsystemUtils/Private/Tests/TestCloudInterface.cpp
new file mode 100644
--- /dev/null
+++ b/Engine/Plugins/Online/OnlineSubsystemUtils/Source/OnlineSubsystemUtils/Private/Tests/TestCloudInterface.cpp
@@ -0,0 +1,415 @@
+// Copyright 1998-2016 Epic Games, Inc. All Rights Reserved.
+
+#include "OnlineSubsystemUtilsPrivatePCH.h"
+#include "TestCloudInterface.h"
+#include "OnlineIdentityInterface.h"
+#include "ModuleManager.h"
+
+#if WITH_DEV_AUTOMATION_TESTS
+
+void FTestCloudInterface::Test(UWorld* InWorld)
+{
+	IOnlineSubsystem* OnlineSub = Online::GetSubsystem(InWorld, FName(*Subsystem));
+	check(OnlineSub); 
+
+	if (OnlineSub->GetIdentityInterface().IsValid())
+	{
+		UserId = OnlineSub->GetIdentityInterface()->GetUniquePlayerId(0);
+	}
+	if (UserId.IsValid())
+	{
+		// Cache interfaces
+		UserCloud = OnlineSub->GetUserCloudInterface();
+		if (UserCloud.IsValid())
+		{
+			// Setup delegates
+			EnumerationDelegate = FOnEnumerateUserFilesCompleteDelegate::CreateRaw(this, &FTestCloudInterface::OnEnumerateUserFilesComplete);
+			OnWriteUserCloudFileCompleteDelegate = FOnWriteUserFileCompleteDelegate::CreateRaw(this, &FTestCloudInterface::OnWriteUserCloudFileComplete);
+			OnReadEnumeratedUserFilesCompleteDelegate = FOnReadUserFileCompleteDelegate::CreateRaw(this, &FTestCloudInterface::OnReadEnumeratedUserFilesComplete);
+			OnDeleteEnumeratedUserFilesCompleteDelegate = FOnDeleteUserFileCompleteDelegate::CreateRaw(this, &FTestCloudInterface::OnDeleteEnumeratedUserFilesComplete);
+
+			SharedCloud = OnlineSub->GetSharedCloudInterface();
+			if (SharedCloud.IsValid())
+			{
+				OnWriteSharedCloudFileCompleteDelegate = FOnWriteSharedFileCompleteDelegate::CreateRaw(this, &FTestCloudInterface::OnWriteSharedCloudFileComplete);
+				OnReadEnumerateSharedFileCompleteDelegate = FOnReadSharedFileCompleteDelegate::CreateRaw(this, &FTestCloudInterface::OnReadEnumeratedSharedFileCompleteDelegate);
+
+				// Pre-generated content
+				SharedCloud->GetDummySharedHandlesForTest(RandomSharedFileHandles);
+			}
+		}
+		else
+		{
+			UE_LOG(LogOnline, Warning, TEXT("Cloud test failed.  Cloud API not supported."));
+			delete this;
+		}
+	}
+	else
+	{
+		UE_LOG(LogOnline, Warning, TEXT("Cloud test failed.  No logged in user."));
+		delete this;
+	}
+}
+
+bool FTestCloudInterface::Tick( float DeltaTime )
+{
+	if (TestPhase != LastTestPhase)
+	{
+		if (!bOverallSuccess)
+		{
+			UE_LOG(LogOnline, Log, TEXT("Testing failed in phase %d"), LastTestPhase);
+			TestPhase = 10;
+		}
+		if (!SharedCloud.IsValid() && 
+			TestPhase > 4)
+		{
+			UE_LOG(LogOnline, Log, TEXT("Skipping shared cloud tests"));
+			TestPhase = 10;
+		}
+
+		switch(TestPhase)
+		{
+		case 0:
+			// Enumerate all files at the beginning
+			EnumerateUserFiles();
+			break;
+		case 1:
+			// Write out N files for the user of various sizes/names and enumerate
+			// @todo: make the file count configurable on a per-platform basis, since different platforms have different file count limits
+			OnWriteUserCloudFileCompleteDelegateHandle = WriteNUserCloudFiles(*UserId, FString(TEXT("UserCloud.bin")), 15, OnWriteUserCloudFileCompleteDelegate);
+			break;
+		case 2:
+			// Read in N files for the user and enumerate
+			OnReadEnumeratedUserFilesCompleteDelegateHandle = ReadEnumeratedUserFiles(OnReadEnumeratedUserFilesCompleteDelegate);
+			break;
+		case 3: 
+			// Forget the N files for the user and enumerate (should still remain locally)
+			OnDeleteEnumeratedUserFilesCompleteDelegateHandle = DeleteEnumeratedUserFiles(true, false, OnDeleteEnumeratedUserFilesCompleteDelegate);
+			break;
+		case 4:
+			// Delete the N files for the user and enumerate (should be removed permanently)
+			OnDeleteEnumeratedUserFilesCompleteDelegateHandle = DeleteEnumeratedUserFiles(false, true, OnDeleteEnumeratedUserFilesCompleteDelegate);
+			break;
+		case 5:
+			// Write out N files for the user and share with the cloud and enumerate
+			OnWriteSharedCloudFileCompleteDelegateHandle = WriteNSharedCloudFiles(*UserId, FString(TEXT("SharedCloud.bin")), 15, OnWriteSharedCloudFileCompleteDelegate);
+			break;
+		case 6:
+			// Read in N files given the shared handles that this user generated
+			OnReadEnumerateSharedFileCompleteDelegateHandle = ReadEnumeratedSharedFiles(false, OnReadEnumerateSharedFileCompleteDelegate);
+			break;
+		case 7:
+			// Delete, marking all the files as only forgotten
+			OnDeleteEnumeratedUserFilesCompleteDelegateHandle = DeleteEnumeratedUserFiles(true, false, OnDeleteEnumeratedUserFilesCompleteDelegate);
+			break;
+		case 8:
+			// Try to read in the N random files from their shared handles
+			OnReadEnumerateSharedFileCompleteDelegateHandle = ReadEnumeratedSharedFiles(true, OnReadEnumerateSharedFileCompleteDelegate);
+			break;
+		case 9:
+			// Delete, marking all the files for permanent removal
+			OnDeleteEnumeratedUserFilesCompleteDelegateHandle = DeleteEnumeratedUserFiles(true, true, OnDeleteEnumeratedUserFilesCompleteDelegate);
+			break;	
+		case 10:
+			bOverallSuccess = bOverallSuccess && Cleanup();
+			UE_LOG(LogOnline, Log, TEXT("TESTING COMPLETE Success:%s!"), bOverallSuccess ? TEXT("true") : TEXT("false"));
+			delete this;
+			return false;
+		}
+		LastTestPhase = TestPhase;
+	}
+	return true;
+}
+
+bool FTestCloudInterface::Cleanup()
+{
+	bool bSuccess = true;
+	if (UserCloud.IsValid() &&
+		!UserCloud->ClearFiles(*UserId))
+	{
+		UE_LOG(LogOnline, Log, TEXT("Failed to cleanup user files"));
+		bSuccess = false;
+	}
+	if (SharedCloud.IsValid() &&
+		!SharedCloud->ClearSharedFiles())
+	{
+		UE_LOG(LogOnline, Log, TEXT("Failed to cleanup shared files"));
+		bSuccess = false;
+	}
+
+	return bSuccess;
+}
+
+void FTestCloudInterface::EnumerateUserFiles()
+{
+	EnumerationDelegateHandle = UserCloud->AddOnEnumerateUserFilesCompleteDelegate_Handle(EnumerationDelegate);
+	UserCloud->EnumerateUserFiles(*UserId);
+}
+
+void FTestCloudInterface::OnEnumerateUserFilesComplete(bool bWasSuccessful, const FUniqueNetId& InUserId)
+{
+	UE_LOG(LogOnline, Log, TEXT("------------------------------------------------"));
+	UE_LOG(LogOnline, Log, TEXT("OnEnumerateUserFilesComplete Success:%d UserId:%s"), bWasSuccessful, *InUserId.ToDebugString());
+	bOverallSuccess = bOverallSuccess && bWasSuccessful;
+
+	TArray<FCloudFileHeader> UserFiles;
+	UserCloud->GetUserFileList(InUserId, UserFiles);
+	UserCloud->ClearOnEnumerateUserFilesCompleteDelegate_Handle(EnumerationDelegateHandle);
+
+	for (int32 Idx=0; Idx<UserFiles.Num(); Idx++)
+	{
+		UE_LOG(LogOnline, Log, TEXT("\t%d FileName:%s DLName:%s Hash:%s Size:%d"),
+			Idx, *UserFiles[Idx].FileName, *UserFiles[Idx].DLName, *UserFiles[Idx].Hash, UserFiles[Idx].FileSize);
+	}
+
+	// Enumeration always advances the test phase
+	TestPhase++;
+}
+
+void FTestCloudInterface::WriteRandomFile(TArray<uint8>& Buffer, int32 Size)
+{
+	Buffer.Empty(Size);
+	Buffer.AddUninitialized(Size);
+	for (int32 i=0; i<Size; i++)
+	{
+		Buffer[i] = i % 255;
+	}
+}
+
+FDelegateHandle FTestCloudInterface::WriteNUserCloudFiles(const FUniqueNetId& InUserId, const FString& FileNameBase, int32 FileCount, FOnWriteUserFileCompleteDelegate& Delegate)
+{
+	UE_LOG(LogOnline, Log, TEXT("------------------------------------------------"));
+	UE_LOG(LogOnline, Log, TEXT("Writing %d files to the cloud for user %s"), FileCount, *InUserId.ToDebugString());
+
+	FString FileName(FileNameBase);
+
+	WriteUserCloudFileCount = FileCount;
+	FDelegateHandle Result = UserCloud->AddOnWriteUserFileCompleteDelegate_Handle(Delegate);
+
+	TArray<uint8> DummyData;
+	for (int32 FileIdx=0; FileIdx<FileCount; FileIdx++)
+	{
+		// @todo: make the dummy data size configurable on a per-platform basis, since different platforms (e.g. PS4) have different size limits
+		//WriteRandomFile(DummyData, FMath::TruncToInt(FMath::FRandRange(1024, 100*1024)));
+		WriteRandomFile(DummyData, FMath::TruncToInt(FMath::FRandRange(256, 1024)));
+		UserCloud->WriteUserFile(InUserId, FString::Printf(TEXT("%s%d.%s"), *FPaths::GetBaseFilename(FileName), FileIdx, *FPaths::GetExtension(FileName)), DummyData);
+	}
+
+	return Result;
+}
+
+FDelegateHandle FTestCloudInterface::WriteNSharedCloudFiles(const FUniqueNetId& InUserId, const FString& FileNameBase, int32 FileCount, FOnWriteSharedFileCompleteDelegate& Delegate)
+{
+	UE_LOG(LogOnline, Log, TEXT("------------------------------------------------"));
+	UE_LOG(LogOnline, Log, TEXT("Writing %d files to the cloud and sharing for user %s"), FileCount, *InUserId.ToDebugString());
+
+	FString FileName(FileNameBase);
+
+	WriteSharedCloudFileCount = FileCount;
+	FDelegateHandle Result = SharedCloud->AddOnWriteSharedFileCompleteDelegate_Handle(Delegate);
+
+	TArray<uint8> DummyData;
+	for (int32 FileIdx=0; FileIdx<FileCount; FileIdx++)
+	{
+		WriteRandomFile(DummyData, FMath::TruncToInt(FMath::FRandRange(1024, 100*1024)));
+		SharedCloud->WriteSharedFile(InUserId, FString::Printf(TEXT("%s%d.%s"), *FPaths::GetBaseFilename(FileName), FileIdx, *FPaths::GetExtension(FileName)), DummyData);
+	}
+
+	return Result;
+}
+
+void FTestCloudInterface::OnWriteUserCloudFileComplete(bool bWasSuccessful, const FUniqueNetId& InUserId, const FString& FileName)
+{
+	UE_LOG(LogOnline, Log, TEXT("Write user file complete Success:%d UserId:%s FileName:%s"), bWasSuccessful, *InUserId.ToDebugString(), *FileName);
+	
+	bOverallSuccess = bOverallSuccess && bWasSuccessful;
+
+	static int32 NumWrittenCount = 0;
+	NumWrittenCount++;
+	if (NumWrittenCount == WriteUserCloudFileCount)
+	{
+		UE_LOG(LogOnline, Log, TEXT("Write %d User Files Complete!"), NumWrittenCount);
+		UE_LOG(LogOnline, Log, TEXT("------------------------------------------------"));
+		UserCloud->ClearOnWriteUserFileCompleteDelegate_Handle(OnWriteUserCloudFileCompleteDelegateHandle);
+		NumWrittenCount = 0;
+		EnumerateUserFiles();
+	}
+}
+
+FDelegateHandle FTestCloudInterface::ReadEnumeratedUserFiles(FOnReadUserFileCompleteDelegate& Delegate)
+{
+	TArray<FCloudFileHeader> UserFiles;
+	UserCloud->GetUserFileList(*UserId, UserFiles);
+
+	ReadUserFileCount = UserFiles.Num();
+
+	UE_LOG(LogOnline, Log, TEXT("------------------------------------------------"));
+	UE_LOG(LogOnline, Log, TEXT("Reading %d enumerated files for user %s"), ReadUserFileCount, *UserId->ToDebugString());
+
+	FDelegateHandle Result;
+	if (ReadUserFileCount > 0)
+	{
+		Result = UserCloud->AddOnReadUserFileCompleteDelegate_Handle(Delegate);
+		for (int32 Idx=0; Idx < ReadUserFileCount; Idx++)
+		{
+			UE_LOG(LogOnline, Log, TEXT("\tFileName:%s Size:%d"), *UserFiles[Idx].FileName, UserFiles[Idx].FileSize);
+			UserCloud->ReadUserFile(*UserId, UserFiles[Idx].FileName);
+		}
+	}
+	else
+	{
+		EnumerateUserFiles();
+	}
+
+	return Result;
+}
+
+void FTestCloudInterface::OnReadEnumeratedUserFilesComplete(bool bWasSuccessful, const FUniqueNetId& InUserId, const FString& FileName)
+{
+	int32 FileSize = 0;
+	char *buffer = NULL;
+	if (bWasSuccessful)
+	{
+		TArray<uint8> FileContents;
+		bWasSuccessful = UserCloud->GetFileContents(InUserId, FileName, FileContents);
+		FileSize = FileContents.Num();
+		buffer = (char *) FileContents.GetData();
+	}
+
+	UE_LOG(LogOnline, Log, TEXT("Read user file complete Success:%d UserId:%s FileName:%s Size:%d"), bWasSuccessful, *InUserId.ToDebugString(), *FileName, FileSize);
+	bOverallSuccess = bOverallSuccess && bWasSuccessful;
+
+	static int32 NumReadCount = 0;
+	NumReadCount++;
+	if (NumReadCount == ReadUserFileCount)
+	{
+		UE_LOG(LogOnline, Log, TEXT("Read %d User Files Complete!"), NumReadCount);
+		UE_LOG(LogOnline, Log, TEXT("------------------------------------------------"));
+		UserCloud->ClearOnReadUserFileCompleteDelegate_Handle(OnReadEnumeratedUserFilesCompleteDelegateHandle);
+		NumReadCount = 0;
+		EnumerateUserFiles();
+	}
+}
+
+FDelegateHandle FTestCloudInterface::DeleteEnumeratedUserFiles(bool bCloudDelete, bool bLocalDelete, FOnDeleteUserFileCompleteDelegate& Delegate)
+{
+	TArray<FCloudFileHeader> UserFiles;
+	UserCloud->GetUserFileList(*UserId, UserFiles);
+
+	DeleteUserFileCount = UserFiles.Num();
+
+	UE_LOG(LogOnline, Log, TEXT("------------------------------------------------"));
+	UE_LOG(LogOnline, Log, TEXT("Deleting %d files from the cloud for user %s CLOUD: %s LOCAL: %s"), DeleteUserFileCount, *UserId->ToDebugString(), bCloudDelete ? TEXT("true") : TEXT("false"), bLocalDelete ? TEXT("true") : TEXT("false"));
+	FDelegateHandle Result;
+	if (DeleteUserFileCount > 0)
+	{
+		Result = UserCloud->AddOnDeleteUserFileCompleteDelegate_Handle(Delegate);
+		for (int32 Idx=0; Idx < DeleteUserFileCount; Idx++)
+		{
+			UE_LOG(LogOnline, Log, TEXT("\tFileName:%s Size:%d"), *UserFiles[Idx].FileName, UserFiles[Idx].FileSize);
+			UserCloud->DeleteUserFile(*UserId, UserFiles[Idx].FileName, bCloudDelete, bLocalDelete);
+		}
+	}
+	else
+	{
+		EnumerateUserFiles();
+	}
+
+	return Result;
+}
+
+void FTestCloudInterface::OnDeleteEnumeratedUserFilesComplete(bool bWasSuccessful, const FUniqueNetId& InUserId, const FString& FileName)
+{
+	UE_LOG(LogOnline, Log, TEXT("Delete user file complete Success:%d UserId:%s FileName:%s"), bWasSuccessful, *InUserId.ToDebugString(), *FileName);
+	bOverallSuccess = bOverallSuccess && bWasSuccessful;
+
+	static int32 NumDeletedCount = 0;
+	NumDeletedCount++;
+	if (NumDeletedCount == DeleteUserFileCount)
+	{
+		UE_LOG(LogOnline, Log, TEXT("Delete %d User Files Complete!"), NumDeletedCount);
+		UE_LOG(LogOnline, Log, TEXT("------------------------------------------------"));
+		UserCloud->ClearOnDeleteUserFileCompleteDelegate_Handle(OnDeleteEnumeratedUserFilesCompleteDelegateHandle);
+		NumDeletedCount = 0;
+		EnumerateUserFiles();
+	}
+}
+
+void FTestCloudInterface::OnWriteSharedCloudFileComplete(bool bWasSuccessful, const FUniqueNetId& InUserId, const FString& FileName, const TSharedRef<FSharedContentHandle>& SharedHandle)
+{
+	UE_LOG(LogOnline, Log, TEXT("Write shared file complete Success:%d UserId:%s FileName:%s SharedHandle:%s"), bWasSuccessful, *InUserId.ToDebugString(), *FileName, *SharedHandle->ToDebugString());
+	bOverallSuccess = bOverallSuccess && bWasSuccessful;
+
+	CloudFileHandles.Add(SharedHandle);
+
+	static int32 NumWrittenCount = 0;
+	NumWrittenCount++;
+	if (NumWrittenCount == WriteSharedCloudFileCount)
+	{
+		UE_LOG(LogOnline, Log, TEXT("Write %d Shared Files Complete!"), NumWrittenCount);
+		UE_LOG(LogOnline, Log, TEXT("------------------------------------------------"));
+		SharedCloud->ClearOnWriteSharedFileCompleteDelegate_Handle(OnWriteSharedCloudFileCompleteDelegateHandle);
+		NumWrittenCount = 0;
+		// Enumeration will put the newly shared files into this user's list
+		EnumerateUserFiles();
+	}
+}
+
+FDelegateHandle FTestCloudInterface::ReadEnumeratedSharedFiles(bool bUseRandom, FOnReadSharedFileCompleteDelegate& Delegate)
+{
+	if (bUseRandom)
+	{
+		ReadSharedFileCount = RandomSharedFileHandles.Num();
+	}
+	else
+	{
+		ReadSharedFileCount = CloudFileHandles.Num();
+	}
+	
+	UE_LOG(LogOnline, Log, TEXT("------------------------------------------------"));
+	UE_LOG(LogOnline, Log, TEXT("Reading %d enumerated shared files"), ReadSharedFileCount);
+
+	FDelegateHandle Result;
+	if (ReadSharedFileCount > 0)
+	{
+		Result = SharedCloud->AddOnReadSharedFileCompleteDelegate_Handle(Delegate);
+		for (int32 Idx=0; Idx < ReadSharedFileCount; Idx++)
+		{
+			if (bUseRandom)
+			{
+				UE_LOG(LogOnline, Log, TEXT("\tHandle:%s"), *RandomSharedFileHandles[Idx]->ToDebugString());
+				SharedCloud->ReadSharedFile(*RandomSharedFileHandles[Idx]);
+			}
+			else
+			{
+				UE_LOG(LogOnline, Log, TEXT("\tHandle:%s"), *CloudFileHandles[Idx]->ToDebugString());
+				SharedCloud->ReadSharedFile(*CloudFileHandles[Idx]);
+			}
+		}
+	}
+	else
+	{
+		EnumerateUserFiles();
+	}
+
+	return Result;
+}
+
+void FTestCloudInterface::OnReadEnumeratedSharedFileCompleteDelegate(bool bWasSuccessful, const FSharedContentHandle& SharedHandle)
+{
+	UE_LOG(LogOnline, Log, TEXT("Read shared file complete Success:%d SharedHandle:%s"), bWasSuccessful, *SharedHandle.ToDebugString());
+	bOverallSuccess = bOverallSuccess && bWasSuccessful;
+
+	static int32 NumReadCount = 0;
+	NumReadCount++;
+	if (NumReadCount == ReadSharedFileCount)
+	{
+		UE_LOG(LogOnline, Log, TEXT("Read %d Shared Files Complete!"), NumReadCount);
+		UE_LOG(LogOnline, Log, TEXT("------------------------------------------------"));
+		SharedCloud->ClearOnReadSharedFileCompleteDelegate_Handle(OnReadEnumerateSharedFileCompleteDelegateHandle);
+		NumReadCount = 0;
+		EnumerateUserFiles();
+	}
+}
+
+#endif //WITH_DEV_AUTOMATION_TESTS
\ No newline at end of file
diff --git a/Engine/Plugins/Online/OnlineSubsystemUtils/Source/OnlineSubsystemUtils/Private/Tests/TestKeyValuePairs.cpp b/Engine/Plugins/Online/OnlineSubsystemUtils/Source/OnlineSubsystemUtils/Private/Tests/TestKeyValuePairs.cpp
new file mode 100644
--- /dev/null
+++ b/Engine/Plugins/Online/OnlineSubsystemUtils/Source/OnlineSubsystemUtils/Private/Tests/TestKeyValuePairs.cpp
@@ -0,0 +1,213 @@
+// Copyright 1998-2016 Epic Games, Inc. All Rights Reserved.
+
+#include "OnlineSubsystemUtilsPrivatePCH.h"
+
+#if WITH_DEV_AUTOMATION_TESTS
+
+/** Simple test cases for key value pair code */
+void TestKeyValuePairs()
+{
+	bool bSuccess = true;
+
+	TArray<uint8> TestData;
+	TestData.Add((uint8)200);
+
+ 	typedef FOnlineKeyValuePairs<FName, FVariantData> TestDataType;
+	TestDataType TestKeyValuePairs;
+ 
+ 	// Test Templates
+ 	TestKeyValuePairs.Add(TEXT("INTValue"), FVariantData(512)); 
+ 	TestKeyValuePairs.Add(TEXT("FLOATValue"), FVariantData(512.0f)); 
+ 	TestKeyValuePairs.Add(TEXT("QWORDValue"), FVariantData((uint64)512)); 
+ 	TestKeyValuePairs.Add(TEXT("DOUBLEValue"), FVariantData(512000.0)); 
+ 	TestKeyValuePairs.Add(TEXT("STRINGValue"), FVariantData(TEXT("This Is A Test!"))); 
+ 	TestKeyValuePairs.Add(TEXT("BLOBValue"), FVariantData(TestData)); 
+
+	UE_LOG(LogOnline, Display, TEXT("ConstIterator"));
+	for (TestDataType::TConstIterator It(TestKeyValuePairs); It; ++It)
+	{
+		UE_LOG(LogOnline, Display, TEXT("%s = %s"), *It.Key().ToString(), *It.Value().ToString());
+	}
+
+	UE_LOG(LogOnline, Display, TEXT("Iterator"));
+	for (TestDataType::TIterator It(TestKeyValuePairs); It; ++It)
+	{
+		UE_LOG(LogOnline, Display, TEXT("Iterator %s = %s"), *It.Key().ToString(), *It.Value().ToString());
+	}
+
+	UE_LOG(LogOnline, Display, TEXT("Finding all elements"));
+	if (TestKeyValuePairs.Find(TEXT("INTValue")) == NULL)
+	{
+		bSuccess = false;
+	}
+	if (TestKeyValuePairs.Find(TEXT("FLOATValue")) == NULL)
+	{
+		bSuccess = false;
+	}
+	if (TestKeyValuePairs.Find(TEXT("QWORDValue")) == NULL)
+	{
+		bSuccess = false;
+	}
+	if (TestKeyValuePairs.Find(TEXT("DOUBLEValue")) == NULL)
+	{
+		bSuccess = false;
+	}
+	if (TestKeyValuePairs.Find(TEXT("STRINGValue")) == NULL)
+	{
+		bSuccess = false;
+	}
+	if (TestKeyValuePairs.Find(TEXT("BLOBValue")) == NULL)
+	{
+		bSuccess = false;
+	}
+
+	if (!bSuccess)
+	{
+		UE_LOG(LogOnline, Display, TEXT("Not all elements found!"));
+	}
+
+	TestKeyValuePairs.Remove(TEXT("INTValue"));
+	TestKeyValuePairs.Remove(TEXT("BLOBValue"));
+
+	UE_LOG(LogOnline, Display, TEXT("Iterator AFTER removing int32 and Blob elements"));
+	for (TestDataType::TIterator It(TestKeyValuePairs); It; ++It)
+	{
+		UE_LOG(LogOnline, Display, TEXT("Iterator %s = %s"), *It.Key().ToString(), *It.Value().ToString());
+	}
+
+	TestKeyValuePairs.Empty();
+	UE_LOG(LogOnline, Display, TEXT("Iterator AFTER emptying structure"));
+	for (TestDataType::TIterator It(TestKeyValuePairs); It; ++It)
+	{
+		UE_LOG(LogOnline, Display, TEXT("Iterator %s = %s"), *It.Key().ToString(), *It.Value().ToString());
+	}
+
+	bSuccess = TestKeyValuePairs.Num() == 0 ? true : false;
+
+	// Test basic variant data type functionality
+	FVariantData OrigKeyValuePair;
+	FVariantData CopyValue;
+
+	{
+		// Test int32
+		int32 OutValue = 0;
+		int32 TestValue = 5;
+		OrigKeyValuePair.SetValue(TestValue);
+		OrigKeyValuePair.GetValue(OutValue);
+		bSuccess = bSuccess && (OutValue == TestValue);   
+		
+		CopyValue = OrigKeyValuePair; 
+		UE_LOG(LogOnline, Display, TEXT("int32 Test %s == %s"), *OrigKeyValuePair.ToString(), *CopyValue.ToString());
+
+		OrigKeyValuePair.Increment<int32,EOnlineKeyValuePairDataType::Int32>((int32)1);
+		UE_LOG(LogOnline, Display, TEXT("+1 Now is %s"), *OrigKeyValuePair.ToString());
+
+		OrigKeyValuePair.Decrement<int32,EOnlineKeyValuePairDataType::Int32>((int32)1);
+		UE_LOG(LogOnline, Display, TEXT("-1 Now is %s"), *OrigKeyValuePair.ToString());
+
+		bSuccess = bSuccess && (OrigKeyValuePair == CopyValue); 
+		bSuccess = bSuccess && OrigKeyValuePair.FromString(TEXT("5"));
+		UE_LOG(LogOnline, Display, TEXT("int32 Test %s == %s"), *OrigKeyValuePair.ToString(), *CopyValue.ToString());
+	}
+
+	{
+		// Test float
+		float OutValue = 0.0f;
+		float TestValue = 5.0f;
+		OrigKeyValuePair.SetValue(TestValue);
+		OrigKeyValuePair.GetValue(OutValue);
+		bSuccess = bSuccess && (OutValue == TestValue);
+
+		CopyValue = OrigKeyValuePair; 
+		UE_LOG(LogOnline, Display, TEXT("float Test %s == %s"), *OrigKeyValuePair.ToString(), *CopyValue.ToString());
+
+		OrigKeyValuePair.Increment<float,EOnlineKeyValuePairDataType::Float>(1.0f);
+		UE_LOG(LogOnline, Display, TEXT("+1 Now is %s"), *OrigKeyValuePair.ToString());
+
+		OrigKeyValuePair.Decrement<float,EOnlineKeyValuePairDataType::Float>(1.0f);
+		UE_LOG(LogOnline, Display, TEXT("-1 Now is %s"), *OrigKeyValuePair.ToString());
+
+		bSuccess = bSuccess && (OrigKeyValuePair == CopyValue); 
+		bSuccess = bSuccess && OrigKeyValuePair.FromString(TEXT("5.0"));
+		UE_LOG(LogOnline, Display, TEXT("float Test %s == %s"), *OrigKeyValuePair.ToString(), *CopyValue.ToString());
+	}
+
+	{
+		// Test double
+		double OutValue = 0.0;
+		double TestValue = 5.0;
+		OrigKeyValuePair.SetValue(TestValue);
+		OrigKeyValuePair.GetValue(OutValue);
+		bSuccess = bSuccess && (OutValue == TestValue);
+
+		CopyValue = OrigKeyValuePair; 
+		UE_LOG(LogOnline, Display, TEXT("double Test %s == %s"), *OrigKeyValuePair.ToString(), *CopyValue.ToString());
+
+		OrigKeyValuePair.Increment<double,EOnlineKeyValuePairDataType::Double>(1.0);
+		UE_LOG(LogOnline, Display, TEXT("+1 Now is %s"), *OrigKeyValuePair.ToString());
+
+		OrigKeyValuePair.Decrement<double,EOnlineKeyValuePairDataType::Double>(1.0);
+		UE_LOG(LogOnline, Display, TEXT("-1 Now is %s"), *OrigKeyValuePair.ToString());
+
+		bSuccess = bSuccess && (OrigKeyValuePair == CopyValue); 
+		bSuccess = bSuccess && OrigKeyValuePair.FromString(TEXT("5.0"));
+		UE_LOG(LogOnline, Display, TEXT("double Test %s == %s"), *OrigKeyValuePair.ToString(), *CopyValue.ToString());
+	}
+
+	{
+		// Test uint64
+		uint64 OutValue = 0;
+		uint64 TestValue = 524288;
+		OrigKeyValuePair.SetValue(TestValue);
+		OrigKeyValuePair.GetValue(OutValue);
+		bSuccess = bSuccess && (OutValue == TestValue);
+
+		CopyValue = OrigKeyValuePair; 
+		UE_LOG(LogOnline, Display, TEXT("uint64 Test %s == %s"), *OrigKeyValuePair.ToString(), *CopyValue.ToString()); 
+
+		OrigKeyValuePair.Increment<uint64,EOnlineKeyValuePairDataType::Int64>(1.0);
+		UE_LOG(LogOnline, Display, TEXT("+1 Now is %s"), *OrigKeyValuePair.ToString());
+
+		OrigKeyValuePair.Decrement<uint64,EOnlineKeyValuePairDataType::Int64>(1.0);
+		UE_LOG(LogOnline, Display, TEXT("-1 Now is %s"), *OrigKeyValuePair.ToString());
+
+		bSuccess = bSuccess && (OrigKeyValuePair == CopyValue);
+		bSuccess = bSuccess && OrigKeyValuePair.FromString(TEXT("524288"));
+		UE_LOG(LogOnline, Display, TEXT("uint64 Test %s == %s"), *OrigKeyValuePair.ToString(), *CopyValue.ToString());
+	}
+
+	{
+		// Test String
+		FString OutValue;
+		FString TestValue(TEXT("This is a test!"));
+		OrigKeyValuePair.SetValue(TestValue);
+		OrigKeyValuePair.GetValue(OutValue);
+		bSuccess = bSuccess && (OutValue == TestValue);
+
+		CopyValue = OrigKeyValuePair; 
+		bSuccess = bSuccess && (OrigKeyValuePair == CopyValue);
+		UE_LOG(LogOnline, Display, TEXT("STRING Test %s == %s"), *OrigKeyValuePair.ToString(), *CopyValue.ToString()); 
+	}
+
+	{
+		// Test Blob
+		TArray<uint8> OutValue;
+		TArray<uint8> TestValue;
+		for (int32 i=0; i<512; i++)
+		{
+			TestValue.Add(FMath::Rand() % 255);
+		}
+
+		OrigKeyValuePair.SetValue(TestValue);
+		OrigKeyValuePair.GetValue(OutValue);
+		bSuccess = bSuccess && (OutValue == TestValue);
+
+		CopyValue = OrigKeyValuePair; 
+		bSuccess = bSuccess && (OrigKeyValuePair == CopyValue); 
+		UE_LOG(LogOnline, Display, TEXT("BLOB Test %s == %s"), *OrigKeyValuePair.ToString(), *CopyValue.ToString());
+	}
+
+	UE_LOG(LogOnline, Warning, TEXT("KeyValuePairTest: %s!"), bSuccess ? TEXT("PASSED") : TEXT("FAILED"));
+}
+
+#endif //WITH_DEV_AUTOMATION_TESTS
\ No newline at end of file
diff --git a/Engine/Plugins/Online/OnlineSubsystemUtils/Source/OnlineSubsystemUtils/Private/Tests/TestLeaderboardInterface.cpp b/Engine/Plugins/Online/OnlineSubsystemUtils/Source/OnlineSubsystemUtils/Private/Tests/TestLeaderboardInterface.cpp
new file mode 100644
--- /dev/null
+++ b/Engine/Plugins/Online/OnlineSubsystemUtils/Source/OnlineSubsystemUtils/Private/Tests/TestLeaderboardInterface.cpp
@@ -0,0 +1,152 @@
+// Copyright 1998-2016 Epic Games, Inc. All Rights Reserved.
+
+#include "OnlineSubsystemUtilsPrivatePCH.h"
+#include "TestLeaderboardInterface.h"
+#include "OnlineIdentityInterface.h"
+#include "ModuleManager.h"
+
+#if WITH_DEV_AUTOMATION_TESTS
+
+/**
+ *	Example of a leaderboard write object
+ */
+class TestLeaderboardWrite : public FOnlineLeaderboardWrite
+{
+public:
+	TestLeaderboardWrite()
+	{
+		// Default properties
+		new (LeaderboardNames) FName(TEXT("TestLeaderboard"));
+		RatedStat = "TestIntStat1";
+		DisplayFormat = ELeaderboardFormat::Number;
+		SortMethod = ELeaderboardSort::Descending;
+		UpdateMethod = ELeaderboardUpdateMethod::KeepBest;
+	}
+};
+
+/**
+ *	Example of a leaderboard read object
+ */
+class TestLeaderboardRead : public FOnlineLeaderboardRead
+{
+public:
+	TestLeaderboardRead()
+	{
+		// Default properties
+		LeaderboardName = FName(TEXT("TestLeaderboard"));
+		SortedColumn = "TestIntStat1";
+
+		// Define default columns
+		new (ColumnMetadata) FColumnMetaData("TestIntStat1", EOnlineKeyValuePairDataType::Int32);
+		new (ColumnMetadata) FColumnMetaData("TestFloatStat1", EOnlineKeyValuePairDataType::Float);
+	}
+};
+
+void FTestLeaderboardInterface::Test(UWorld* InWorld)
+{
+	IOnlineSubsystem* OnlineSub = Online::GetSubsystem(InWorld, FName(*Subsystem));
+	check(OnlineSub); 
+
+	if (OnlineSub->GetIdentityInterface().IsValid())
+	{
+		UserId = OnlineSub->GetIdentityInterface()->GetUniquePlayerId(0);
+	}
+
+	// Cache interfaces
+	Leaderboards = OnlineSub->GetLeaderboardsInterface();
+	check(Leaderboards.IsValid());
+
+	// Define delegates
+	LeaderboardFlushDelegate = FOnLeaderboardFlushCompleteDelegate::CreateRaw(this, &FTestLeaderboardInterface::OnLeaderboardFlushComplete);
+	LeaderboardReadCompleteDelegate = FOnLeaderboardReadCompleteDelegate::CreateRaw(this, &FTestLeaderboardInterface::OnLeaderboardReadComplete);
+}
+
+void FTestLeaderboardInterface::WriteLeaderboards()
+{
+	TestLeaderboardWrite WriteObject;
+	
+	// Set some data
+	WriteObject.SetIntStat("TestIntStat1", 50);
+	WriteObject.SetFloatStat("TestFloatStat1", 99.5f);
+
+	// Write it to the buffers
+	Leaderboards->WriteLeaderboards(TEXT("TEST"), *UserId, WriteObject);
+	TestPhase++;
+}
+
+void FTestLeaderboardInterface::OnLeaderboardFlushComplete(FName SessionName, bool bWasSuccessful)
+{
+	UE_LOG(LogOnline, Verbose, TEXT("OnLeaderboardFlushComplete Session: %s bWasSuccessful: %d"), *SessionName.ToString(), bWasSuccessful);
+	bOverallSuccess = bOverallSuccess && bWasSuccessful;
+
+	Leaderboards->ClearOnLeaderboardFlushCompleteDelegate_Handle(LeaderboardFlushDelegateHandle);
+	TestPhase++;
+}
+
+void FTestLeaderboardInterface::FlushLeaderboards()
+{
+	LeaderboardFlushDelegateHandle = Leaderboards->AddOnLeaderboardFlushCompleteDelegate_Handle(LeaderboardFlushDelegate);
+	Leaderboards->FlushLeaderboards(TEXT("TEST"));
+}
+
+void FTestLeaderboardInterface::OnLeaderboardReadComplete(bool bWasSuccessful)
+{
+	UE_LOG(LogOnline, Verbose, TEXT("OnLeaderboardReadComplete bWasSuccessful: %d"), bWasSuccessful);
+	bOverallSuccess = bOverallSuccess && bWasSuccessful;
+
+	for (int32 RowIdx = 0; RowIdx < ReadObject->Rows.Num(); ++RowIdx)
+	{
+		const FOnlineStatsRow& StatsRow = ReadObject->Rows[RowIdx];
+		UE_LOG(LogOnline, Log, TEXT("Leaderboard stats for: Nickname = %s, Rank = %d"), *StatsRow.NickName, StatsRow.Rank);
+
+		for (FStatsColumnArray::TConstIterator It(StatsRow.Columns); It; ++It)
+		{
+			UE_LOG(LogOnline, Log, TEXT("  %s = %s"), *It.Key().ToString(), *It.Value().ToString());
+		}
+	}
+
+	Leaderboards->ClearOnLeaderboardReadCompleteDelegate_Handle(LeaderboardReadCompleteDelegateHandle);
+	TestPhase++;
+}
+
+void FTestLeaderboardInterface::ReadLeaderboards()
+{
+	ReadObject = MakeShareable(new TestLeaderboardRead());
+	FOnlineLeaderboardReadRef ReadObjectRef = ReadObject.ToSharedRef();
+
+	LeaderboardReadCompleteDelegateHandle = Leaderboards->AddOnLeaderboardReadCompleteDelegate_Handle(LeaderboardReadCompleteDelegate);
+	Leaderboards->ReadLeaderboardsForFriends(0, ReadObjectRef);
+}
+
+bool FTestLeaderboardInterface::Tick( float DeltaTime )
+{
+	if (TestPhase != LastTestPhase)
+	{
+		LastTestPhase = TestPhase;
+		if (!bOverallSuccess)
+		{
+			UE_LOG(LogOnline, Log, TEXT("Testing failed in phase %d"), LastTestPhase);
+			TestPhase = 3;
+		}
+
+		switch(TestPhase)
+		{
+		case 0:
+			WriteLeaderboards();
+			break;
+		case 1:
+			FlushLeaderboards();
+			break;
+		case 2:
+			ReadLeaderboards();
+			break;
+		case 3:
+			UE_LOG(LogOnline, Log, TEXT("TESTING COMPLETE Success:%s!"), bOverallSuccess ? TEXT("true") : TEXT("false"));
+			delete this;
+			return false;
+		}
+	}
+	return true;
+}
+
+#endif //WITH_DEV_AUTOMATION_TESTS
\ No newline at end of file
diff --git a/Engine/Plugins/Online/OnlineSubsystemUtils/Source/OnlineSubsystemUtils/Private/Tests/TestSessionInterface.cpp b/Engine/Plugins/Online/OnlineSubsystemUtils/Source/OnlineSubsystemUtils/Private/Tests/TestSessionInterface.cpp
new file mode 100644
--- /dev/null
+++ b/Engine/Plugins/Online/OnlineSubsystemUtils/Source/OnlineSubsystemUtils/Private/Tests/TestSessionInterface.cpp
@@ -0,0 +1,549 @@
+// Copyright 1998-2016 Epic Games, Inc. All Rights Reserved.
+
+#include "OnlineSubsystemUtilsPrivatePCH.h"
+#include "TestSessionInterface.h"
+#include "OnlineIdentityInterface.h"
+#include "OnlineFriendsInterface.h"
+#include "OnlineExternalUIInterface.h"
+#include "ModuleManager.h"
+#include "Kismet/GameplayStatics.h"
+
+#include "Engine.h"
+
+#if WITH_DEV_AUTOMATION_TESTS
+
+/**
+ *	Example of a hosted session
+ */
+class TestOnlineGameSettings : public FOnlineSessionSettings
+{
+ public:
+ 	TestOnlineGameSettings(bool bTestingLAN = false, bool bTestingPresence = false)
+	{
+ 		NumPublicConnections = 10;
+		NumPrivateConnections = 0;
+		bIsLANMatch = bTestingLAN;
+		bShouldAdvertise = true;
+		bAllowJoinInProgress = true;
+		bAllowInvites = true;
+		bUsesPresence = bTestingPresence;
+		bAllowJoinViaPresence = true;
+		bAllowJoinViaPresenceFriendsOnly = false;
+
+		Set(FName(TEXT("TESTSETTING1")), (int32)5, EOnlineDataAdvertisementType::ViaOnlineService);
+		Set(FName(TEXT("TESTSETTING2")), (float)5.0f, EOnlineDataAdvertisementType::ViaOnlineService);
+		Set(FName(TEXT("TESTSETTING3")), FString(TEXT("Hello")), EOnlineDataAdvertisementType::ViaOnlineService);
+		Set(FName(TEXT("TESTSETTING4")), FString(TEXT("Test4")), EOnlineDataAdvertisementType::ViaPingOnly);
+		Set(FName(TEXT("TESTSETTING5")), FString(TEXT("Test5")), EOnlineDataAdvertisementType::ViaPingOnly);
+
+ 	}
+ 	
+ 	virtual ~TestOnlineGameSettings()
+ 	{
+ 
+ 	}
+
+	void AddWorldSettings(UWorld* InWorld)
+	{
+		if (InWorld)
+		{
+			const FString MapName = InWorld->GetMapName();
+			Set(SETTING_MAPNAME, MapName, EOnlineDataAdvertisementType::ViaOnlineService);
+
+			AGameMode const* const GameMode = InWorld->GetAuthGameMode();
+			if (GameMode != NULL)
+			{
+				// Bot count
+				Set(SETTING_NUMBOTS, GameMode->NumBots, EOnlineDataAdvertisementType::ViaOnlineService);
+
+				// Game type
+				FString GameModeStr = GameMode->GetClass()->GetName();
+				Set(SETTING_GAMEMODE, GameModeStr, EOnlineDataAdvertisementType::ViaOnlineService);
+			}
+		}
+	}
+ };
+
+/**
+ *	Example of a session search query
+ */
+class TestOnlineSearchSettings : public FOnlineSessionSearch
+{
+public:
+	TestOnlineSearchSettings(bool bSearchingLAN = false, bool bSearchingPresence = false)
+	{
+		bIsLanQuery = bSearchingLAN;
+		MaxSearchResults = 10;
+		PingBucketSize = 50;
+
+		QuerySettings.Set(FName(TEXT("TESTSETTING1")), (int32)5, EOnlineComparisonOp::Equals);
+		QuerySettings.Set(FName(TEXT("TESTSETTING2")), (float)5.0f, EOnlineComparisonOp::Equals);
+		QuerySettings.Set(FName(TEXT("TESTSETTING3")), FString(TEXT("Hello")), EOnlineComparisonOp::Equals);
+
+		if (bSearchingPresence)
+		{
+			QuerySettings.Set(SEARCH_PRESENCE, true, EOnlineComparisonOp::Equals);
+		}
+	}
+
+	virtual ~TestOnlineSearchSettings()
+	{
+
+	}
+};
+
+void FTestSessionInterface::Test(UWorld* InWorld, bool bTestLAN, bool bIsPresence)
+{
+	IOnlineSubsystem* OnlineSub = Online::GetSubsystem(InWorld, FName(*Subsystem));
+	check(OnlineSub); 
+
+	World = InWorld;
+	GEngine->OnWorldDestroyed().AddRaw(this, &FTestSessionInterface::WorldDestroyed);
+
+	if (OnlineSub->GetIdentityInterface().IsValid())
+	{
+		UserId = OnlineSub->GetIdentityInterface()->GetUniquePlayerId(0);
+	}
+	
+	// Cache interfaces
+	Identity = OnlineSub->GetIdentityInterface();
+	SessionInt = OnlineSub->GetSessionInterface();
+	check(SessionInt.IsValid());
+	Friends = OnlineSub->GetFriendsInterface();
+
+	// Define delegates
+	OnCreateSessionCompleteDelegate = FOnCreateSessionCompleteDelegate::CreateRaw(this, &FTestSessionInterface::OnCreateSessionComplete);
+	OnStartSessionCompleteDelegate = FOnStartSessionCompleteDelegate::CreateRaw(this, &FTestSessionInterface::OnStartSessionComplete);
+	OnEndSessionCompleteDelegate = FOnEndSessionCompleteDelegate::CreateRaw(this, &FTestSessionInterface::OnEndSessionComplete);
+	OnDestroySessionCompleteDelegate = FOnDestroySessionCompleteDelegate::CreateRaw(this, &FTestSessionInterface::OnDestroySessionComplete);
+
+	OnUpdateSessionCompleteDelegate = FOnUpdateSessionCompleteDelegate::CreateRaw(this, &FTestSessionInterface::OnUpdateSessionComplete);
+
+	OnJoinSessionCompleteDelegate = FOnJoinSessionCompleteDelegate::CreateRaw(this, &FTestSessionInterface::OnJoinSessionComplete);
+	OnEndForJoinSessionCompleteDelegate = FOnEndSessionCompleteDelegate::CreateRaw(this, &FTestSessionInterface::OnEndForJoinSessionComplete);
+	OnDestroyForJoinSessionCompleteDelegate = FOnDestroySessionCompleteDelegate::CreateRaw(this, &FTestSessionInterface::OnDestroyForJoinSessionComplete);
+
+	OnFindFriendSessionCompleteDelegate = FOnFindFriendSessionCompleteDelegate::CreateRaw(this, &FTestSessionInterface::OnFindFriendSessionComplete);
+
+	OnRegisterPlayersCompleteDelegate = FOnRegisterPlayersCompleteDelegate::CreateRaw(this, &FTestSessionInterface::OnRegisterPlayerComplete);
+	OnUnregisterPlayersCompleteDelegate = FOnRegisterPlayersCompleteDelegate::CreateRaw(this, &FTestSessionInterface::OnUnregisterPlayerComplete);
+
+	OnFindSessionsCompleteDelegate = FOnFindSessionsCompleteDelegate::CreateRaw(this, &FTestSessionInterface::OnFindSessionsComplete);
+	OnCancelFindSessionsCompleteDelegate = FOnCancelFindSessionsCompleteDelegate::CreateRaw(this, &FTestSessionInterface::OnCancelFindSessionsComplete);
+
+	// Read friends list and cache it
+	if (Friends.IsValid())
+	{
+		Friends->ReadFriendsList(0, EFriendsLists::ToString(EFriendsLists::Default), FOnReadFriendsListComplete::CreateRaw(this, &FTestSessionInterface::OnReadFriendsListComplete));
+	}
+
+	// Setup sessions
+	if (bIsHost)
+	{
+ 		HostSettings = MakeShareable(new TestOnlineGameSettings(bTestLAN, bIsPresence));
+		HostSettings->AddWorldSettings(InWorld);
+ 		OnCreateSessionCompleteDelegateHandle = SessionInt->AddOnCreateSessionCompleteDelegate_Handle(OnCreateSessionCompleteDelegate);
+ 		SessionInt->CreateSession(0, GameSessionName, *HostSettings);
+	}
+	else
+	{
+		SearchSettings = MakeShareable(new TestOnlineSearchSettings(bTestLAN, bIsPresence));
+		TSharedRef<FOnlineSessionSearch> SearchSettingsRef = SearchSettings.ToSharedRef();
+
+		OnFindSessionsCompleteDelegateHandle = SessionInt->AddOnFindSessionsCompleteDelegate_Handle(OnFindSessionsCompleteDelegate);
+		SessionInt->FindSessions(0, SearchSettingsRef);
+	}
+
+	// Always on the lookout for invite acceptance (via actual invite or join in an external client)
+	OnSessionUserInviteAcceptedDelegate = FOnSessionUserInviteAcceptedDelegate::CreateRaw(this, &FTestSessionInterface::OnSessionUserInviteAccepted);
+	OnSessionUserInviteAcceptedDelegateHandle = SessionInt->AddOnSessionUserInviteAcceptedDelegate_Handle(OnSessionUserInviteAcceptedDelegate);
+}
+
+void FTestSessionInterface::ClearDelegates()
+{
+	SessionInt->ClearOnSessionUserInviteAcceptedDelegate_Handle(OnSessionUserInviteAcceptedDelegateHandle);
+	GEngine->OnWorldDestroyed().RemoveAll(this);
+}
+
+void FTestSessionInterface::WorldDestroyed( UWorld* InWorld )
+{
+	if (InWorld == World)
+	{
+		World = NULL;
+	}
+}
+
+void FTestSessionInterface::OnReadFriendsListComplete(int32 LocalUserNum, bool bWasSuccessful, const FString& ListName, const FString& ErrorStr)
+{
+	UE_LOG(LogOnline, Verbose, TEXT("OnReadFriendsListComplete LocalUserNum: %d bSuccess: %d %s"), LocalUserNum, bWasSuccessful, *ErrorStr);
+	if (bWasSuccessful)
+	{
+		Friends->GetFriendsList(LocalUserNum, ListName, FriendsCache);
+	}
+}
+
+void FTestSessionInterface::OnSessionUserInviteAccepted(bool bWasSuccessful, int32 ControllerId, TSharedPtr<const FUniqueNetId> InUserId, const FOnlineSessionSearchResult& SearchResult)
+{
+	UE_LOG(LogOnline, Verbose, TEXT("OnSessionInviteAccepted ControllerId: %d bSuccess: %d"), ControllerId, bWasSuccessful);
+	// Don't clear invite accept delegate
+
+	if (bWasSuccessful)
+	{
+		JoinSession(ControllerId, GameSessionName, SearchResult);
+	}
+}
+
+void FTestSessionInterface::OnEndForJoinSessionComplete(FName SessionName, bool bWasSuccessful)
+{
+	UE_LOG(LogOnline, Verbose, TEXT("OnEndForJoinSessionComplete %s bSuccess: %d"), *SessionName.ToString(), bWasSuccessful);
+	SessionInt->ClearOnEndSessionCompleteDelegate_Handle(OnEndForJoinSessionCompleteDelegateHandle);
+	DestroyExistingSession(SessionName, OnDestroyForJoinSessionCompleteDelegate);
+}
+
+void FTestSessionInterface::EndExistingSession(FName SessionName, FOnEndSessionCompleteDelegate& Delegate)
+{
+	SessionInt->AddOnEndSessionCompleteDelegate_Handle(Delegate);
+	SessionInt->EndSession(SessionName);
+}
+
+void FTestSessionInterface::OnDestroyForJoinSessionComplete(FName SessionName, bool bWasSuccessful)
+{
+	UE_LOG(LogOnline, Verbose, TEXT("OnDestroyForJoinSessionComplete %s bSuccess: %d"), *SessionName.ToString(), bWasSuccessful);
+	SessionInt->ClearOnDestroySessionCompleteDelegate_Handle(OnDestroyForJoinSessionCompleteDelegateHandle);
+	JoinSession(0, SessionName, CachedSessionResult);
+}
+
+void FTestSessionInterface::DestroyExistingSession(FName SessionName, FOnDestroySessionCompleteDelegate& Delegate)
+{
+	SessionInt->AddOnDestroySessionCompleteDelegate_Handle(Delegate);
+	SessionInt->DestroySession(SessionName);
+}
+
+void FTestSessionInterface::OnRegisterPlayerComplete(FName SessionName, const TArray< TSharedRef<const FUniqueNetId> >& Players, bool bWasSuccessful)
+{
+	UE_LOG(LogOnline, Verbose, TEXT("OnRegisterPlayerComplete %s bSuccess: %d"), *SessionName.ToString(), bWasSuccessful);
+	SessionInt->ClearOnRegisterPlayersCompleteDelegate_Handle(OnRegisterPlayersCompleteDelegateHandle);
+}
+
+void FTestSessionInterface::OnUnregisterPlayerComplete(FName SessionName, const TArray< TSharedRef<const FUniqueNetId> >& Players, bool bWasSuccessful)
+{
+	UE_LOG(LogOnline, Verbose, TEXT("OnUnregisterPlayerComplete %s bSuccess: %d"), *SessionName.ToString(), bWasSuccessful);
+	SessionInt->ClearOnUnregisterPlayersCompleteDelegate_Handle(OnUnregisterPlayersCompleteDelegateHandle);
+}
+
+void FTestSessionInterface::JoinSession(int32 ControllerId, FName SessionName, const FOnlineSessionSearchResult& SearchResult)
+{
+	// Clean up existing sessions if applicable
+	EOnlineSessionState::Type SessionState = SessionInt->GetSessionState(SessionName);
+	if (SessionState != EOnlineSessionState::NoSession)
+	{
+		CachedSessionResult = SearchResult;
+		EndExistingSession(SessionName, OnEndForJoinSessionCompleteDelegate);
+	}
+	else
+	{
+		OnJoinSessionCompleteDelegateHandle = SessionInt->AddOnJoinSessionCompleteDelegate_Handle(OnJoinSessionCompleteDelegate);
+		SessionInt->JoinSession(ControllerId, SessionName, SearchResult);
+	}
+}
+
+void FTestSessionInterface::OnCreateSessionComplete(FName SessionName, bool bWasSuccessful)
+{
+	UE_LOG(LogOnline, Verbose, TEXT("OnCreateSessionComplete %s bSuccess: %d"), *SessionName.ToString(), bWasSuccessful);
+	SessionInt->ClearOnCreateSessionCompleteDelegate_Handle(OnCreateSessionCompleteDelegateHandle);
+
+	bWasLastCommandSuccessful = bWasSuccessful;
+}
+
+void FTestSessionInterface::OnStartSessionComplete(FName SessionName, bool bWasSuccessful)
+{
+	UE_LOG(LogOnline, Verbose, TEXT("OnStartSessionComplete %s bSuccess: %d"), *SessionName.ToString(), bWasSuccessful);
+	SessionInt->ClearOnStartSessionCompleteDelegate_Handle(OnStartSessionCompleteDelegateHandle);
+
+	bWasLastCommandSuccessful = bWasSuccessful;
+}
+
+void FTestSessionInterface::OnEndSessionComplete(FName SessionName, bool bWasSuccessful)
+{
+	UE_LOG(LogOnline, Verbose, TEXT("OnEndSessionComplete %s bSuccess: %d"), *SessionName.ToString(), bWasSuccessful);
+	SessionInt->ClearOnEndSessionCompleteDelegate_Handle(OnEndSessionCompleteDelegateHandle);
+
+	bWasLastCommandSuccessful = bWasSuccessful;
+}
+
+void FTestSessionInterface::OnDestroySessionComplete(FName SessionName, bool bWasSuccessful)
+{
+	UE_LOG(LogOnline, Verbose, TEXT("OnDestroySessionComplete %s bSuccess: %d"), *SessionName.ToString(), bWasSuccessful);
+	SessionInt->ClearOnDestroySessionCompleteDelegate_Handle(OnDestroySessionCompleteDelegateHandle);
+}
+
+void FTestSessionInterface::OnUpdateSessionComplete(FName SessionName, bool bWasSuccessful)
+{
+	UE_LOG(LogOnline, Verbose, TEXT("OnUpdateSessionComplete %s bSuccess: %d"), *SessionName.ToString(), bWasSuccessful);
+	SessionInt->ClearOnUpdateSessionCompleteDelegate_Handle(OnUpdateSessionCompleteDelegateHandle);
+}
+
+void FTestSessionInterface::OnJoinSessionComplete(FName SessionName, EOnJoinSessionCompleteResult::Type Result)
+{
+	UE_LOG(LogOnline, Verbose, TEXT("OnJoinSessionComplete %s bSuccess: %d"), *SessionName.ToString(), static_cast<int32>(Result));
+	SessionInt->ClearOnJoinSessionCompleteDelegate_Handle(OnJoinSessionCompleteDelegateHandle);
+
+	if (Result == EOnJoinSessionCompleteResult::Success)
+	{
+		FString URL;
+		if (World && SessionInt->GetResolvedConnectString(SessionName, URL))
+		{
+			APlayerController* PC = UGameplayStatics::GetPlayerController(World, 0);
+			if (PC)
+			{
+				PC->ClientTravel(URL, TRAVEL_Absolute);
+			}
+		}
+		else
+		{
+			UE_LOG(LogOnline, Warning, TEXT("Failed to join session %s"), *SessionName.ToString());
+		}
+	}
+}
+
+void FTestSessionInterface::OnFindFriendSessionComplete(int32 LocalUserNum, bool bWasSuccessful, const FOnlineSessionSearchResult& SearchResult)
+{
+	UE_LOG(LogOnline, Verbose, TEXT("OnFindFriendSessionComplete LocalUserNum: %d bSuccess: %d"), LocalUserNum, bWasSuccessful);
+
+	FDelegateHandle DelegateHandle = OnFindFriendSessionCompleteDelegateHandles.FindRef(LocalUserNum);
+	SessionInt->ClearOnFindFriendSessionCompleteDelegate_Handle(LocalUserNum, DelegateHandle);
+	OnFindFriendSessionCompleteDelegateHandles.Remove(LocalUserNum);
+	if (bWasSuccessful)
+	{
+		if (SearchResult.IsValid())
+		{
+			JoinSession(LocalUserNum, GameSessionName, SearchResult);
+		}
+		else
+		{
+			UE_LOG(LogOnline, Warning, TEXT("Join friend returned no search result."));
+		}
+	}
+}
+
+void FTestSessionInterface::OnFindSessionsComplete(bool bWasSuccessful)
+{
+	UE_LOG(LogOnline, Verbose, TEXT("OnFindSessionsComplete bSuccess: %d"), bWasSuccessful);
+	SessionInt->ClearOnFindSessionsCompleteDelegate_Handle(OnFindSessionsCompleteDelegateHandle);
+
+	UE_LOG(LogOnline, Verbose, TEXT("Num Search Results: %d"), SearchSettings->SearchResults.Num());
+	for (int32 SearchIdx=0; SearchIdx<SearchSettings->SearchResults.Num(); SearchIdx++)
+	{
+		const FOnlineSessionSearchResult& SearchResult = SearchSettings->SearchResults[SearchIdx];
+		DumpSession(&SearchResult.Session);
+	}
+}
+
+void FTestSessionInterface::OnCancelFindSessionsComplete(bool bWasSuccessful)
+{
+	UE_LOG(LogOnline, Verbose, TEXT("OnCancelFindSessionsComplete bSuccess: %d"), bWasSuccessful);
+	SessionInt->ClearOnCancelFindSessionsCompleteDelegate_Handle(OnCancelFindSessionsCompleteDelegateHandle);
+}
+
+bool FTestSessionInterface::Tick(float DeltaTime)
+{
+	if (TestPhase != LastTestPhase)
+	{
+		LastTestPhase = TestPhase;
+		if (!bOverallSuccess)
+		{
+			UE_LOG(LogOnline, Log, TEXT("Testing failed in phase %d"), LastTestPhase);
+			TestPhase = 3;
+		}
+
+		switch(TestPhase)
+		{
+		case 0:
+			break;
+		case 1:
+			UE_LOG(LogOnline, Log, TEXT("TESTING COMPLETE Success:%s!"), bOverallSuccess ? TEXT("true") : TEXT("false"));
+			delete this;
+			return false;
+		}
+	}
+	return true;
+}
+
+bool FTestSessionInterface::Exec(UWorld* InWorld, const TCHAR* Cmd, FOutputDevice& Ar)
+{
+	bool bWasHandled = false;
+
+	int32 LocalUserNum = 0;
+
+	// Ignore any execs that don't start with TESTSESSION
+	if (FParse::Command(&Cmd, TEXT("TESTSESSION")))
+	{
+		// Get optional session name "GAME" otherwise
+		FString Token;
+		FName SessionName = FParse::Value(Cmd, TEXT("Name="), Token) ? FName(*Token) : GameSessionName;
+		if (Token.Len() > 0)
+		{
+			Cmd += FCString::Strlen(TEXT("Name=")) + Token.Len();
+		}
+		
+		if (FParse::Command(&Cmd, TEXT("SEARCH")))
+		{
+			if (SearchSettings.IsValid())
+			{
+				TSharedRef<FOnlineSessionSearch> SearchSettingsRef = SearchSettings.ToSharedRef();
+
+				OnFindSessionsCompleteDelegateHandle = SessionInt->AddOnFindSessionsCompleteDelegate_Handle(OnFindSessionsCompleteDelegate);
+				SessionInt->FindSessions(LocalUserNum, SearchSettingsRef);
+			}
+			bWasHandled = true;
+		}
+		else if (FParse::Command(&Cmd, TEXT("JOIN")))
+		{
+			TCHAR SearchIdxStr[256];
+			if (FParse::Token(Cmd, SearchIdxStr, ARRAY_COUNT(SearchIdxStr), true))
+			{
+				int32 SearchIdx = FCString::Atoi(SearchIdxStr);
+				if (SearchIdx >= 0 && SearchIdx < SearchSettings->SearchResults.Num())
+				{
+					JoinSession(LocalUserNum, SessionName, SearchSettings->SearchResults[SearchIdx]);
+				}
+			}
+			bWasHandled = true;
+		}
+		else if (FParse::Command(&Cmd, TEXT("JOINFRIEND")))
+		{
+			if (FParse::Command(&Cmd, TEXT("LOBBY")))
+			{
+				TCHAR FriendNameStr[256];
+				if (FParse::Token(Cmd, FriendNameStr, ARRAY_COUNT(FriendNameStr), true))
+				{
+					for (int32 FriendIdx=0; FriendIdx<FriendsCache.Num(); FriendIdx++)
+					{
+						const FOnlineFriend& Friend = *FriendsCache[FriendIdx];
+						if (Friend.GetDisplayName() == FriendNameStr)
+						{
+							OnFindFriendSessionCompleteDelegateHandles.Add(LocalUserNum, SessionInt->AddOnFindFriendSessionCompleteDelegate_Handle(LocalUserNum, OnFindFriendSessionCompleteDelegate));
+							SessionInt->FindFriendSession(LocalUserNum, *Friend.GetUserId());
+							break;
+						}
+					}
+				}
+			}
+			else
+			{
+				TCHAR FriendIdStr[256];
+				if (FParse::Token(Cmd, FriendIdStr, ARRAY_COUNT(FriendIdStr), true))
+				{
+					TSharedPtr<const FUniqueNetId> FriendId = Identity->CreateUniquePlayerId((uint8*)FriendIdStr, FCString::Strlen(FriendIdStr));
+					OnFindFriendSessionCompleteDelegateHandles.Add(LocalUserNum, SessionInt->AddOnFindFriendSessionCompleteDelegate_Handle(LocalUserNum, OnFindFriendSessionCompleteDelegate));
+					SessionInt->FindFriendSession(LocalUserNum, *FriendId);
+				}
+			}
+			bWasHandled = true;
+		}
+		else if (FParse::Command(&Cmd, TEXT("CREATE")))
+		{
+			// Mutually exclusive settings
+			bool bTestLAN = FParse::Command(&Cmd, TEXT("LAN")) ? true : false;
+			bool bTestPresence = FParse::Command(&Cmd, TEXT("PRESENCE")) ? true : false;
+
+			HostSettings = MakeShareable(new TestOnlineGameSettings(bTestLAN));
+			HostSettings->bUsesPresence = bTestPresence;
+			HostSettings->AddWorldSettings(InWorld);
+
+			OnCreateSessionCompleteDelegateHandle = SessionInt->AddOnCreateSessionCompleteDelegate_Handle(OnCreateSessionCompleteDelegate);
+			SessionInt->CreateSession(0, SessionName, *HostSettings);
+			bWasHandled = true;
+		}
+		else if (FParse::Command(&Cmd, TEXT("START")))
+		{
+			OnStartSessionCompleteDelegateHandle = SessionInt->AddOnStartSessionCompleteDelegate_Handle(OnStartSessionCompleteDelegate);
+			SessionInt->StartSession(SessionName);
+			bWasHandled = true;
+		}
+		else if (FParse::Command(&Cmd, TEXT("UPDATE")))
+		{
+			bool bUpdateOnline = FParse::Command(&Cmd, TEXT("ONLINE")) ? true : false;
+			OnUpdateSessionCompleteDelegateHandle = SessionInt->AddOnUpdateSessionCompleteDelegate_Handle(OnUpdateSessionCompleteDelegate);
+			HostSettings->Set(FName(TEXT("UPDATESETTING1")), FString(TEXT("Test1")), EOnlineDataAdvertisementType::ViaOnlineService);
+			HostSettings->Set(FName(TEXT("UPDATESETTING2")), 99, EOnlineDataAdvertisementType::ViaOnlineService);
+			SessionInt->UpdateSession(SessionName, *HostSettings, bUpdateOnline);
+			bWasHandled = true;
+		}
+		else if (FParse::Command(&Cmd, TEXT("END")))
+		{
+			EndExistingSession(SessionName, OnEndSessionCompleteDelegate);
+			bWasHandled = true;
+		}
+		else if (FParse::Command(&Cmd, TEXT("DESTROY")))
+		{
+			DestroyExistingSession(SessionName, OnDestroySessionCompleteDelegate);
+			bWasHandled = true;
+		}
+		else if (FParse::Command(&Cmd, TEXT("REGISTER")))
+		{
+			bool bWasInvited = FParse::Command(&Cmd, TEXT("INVITED")) ? true : false;
+			OnRegisterPlayersCompleteDelegateHandle = SessionInt->AddOnRegisterPlayersCompleteDelegate_Handle(OnRegisterPlayersCompleteDelegate);
+			SessionInt->RegisterPlayer(SessionName, *UserId, bWasInvited);
+			bWasHandled = true;
+		}
+		else if (FParse::Command(&Cmd, TEXT("UNREGISTER")))
+		{
+			OnUnregisterPlayersCompleteDelegateHandle = SessionInt->AddOnUnregisterPlayersCompleteDelegate_Handle(OnUnregisterPlayersCompleteDelegate);
+			SessionInt->UnregisterPlayer(SessionName, *UserId);
+			bWasHandled = true;
+		}
+		else if (FParse::Command(&Cmd, TEXT("INVITE")))
+		{
+			if (FParse::Command(&Cmd, TEXT("UI")))
+			{
+				Online::GetSubsystem(InWorld, FName(*Subsystem))->GetExternalUIInterface()->ShowInviteUI(LocalUserNum);
+			}
+			else
+			{
+				FString FriendStr = FParse::Token(Cmd, true);
+				if (!FriendStr.IsEmpty())
+				{
+					bool bFound = false;
+					for (auto Friend : FriendsCache)
+					{
+						if (Friend->GetDisplayName() == FriendStr)
+						{
+							SessionInt->SendSessionInviteToFriend(LocalUserNum, SessionName, *Friend->GetUserId());
+							bFound = true;
+							break;
+						}
+					}
+					if (!bFound)
+					{
+						// use friend str as id instead of display name
+						TSharedPtr<const FUniqueNetId> FriendUserId = Identity->CreateUniquePlayerId(FriendStr);
+						if (FriendUserId.IsValid())
+						{	
+							SessionInt->SendSessionInviteToFriend(LocalUserNum, SessionName, *FriendUserId);
+						}
+					}
+				}
+			}
+
+			bWasHandled = true;
+		}
+		else if (FParse::Command(&Cmd, TEXT("DUMPSESSIONS")))
+		{
+			SessionInt->DumpSessionState();
+			bWasHandled = true;
+		}
+		else if (FParse::Command(&Cmd, TEXT("QUIT")))
+		{
+			UE_LOG(LogOnline, Display, TEXT("Destroying TestSession."));
+			SessionInt->CancelFindSessions();
+			bWasHandled = true;
+			// Make this last
+			delete this;
+		}
+	}
+
+	return bWasHandled;
+}
+
+#endif //WITH_DEV_AUTOMATION_TESTS
\ No newline at end of file
diff --git a/Engine/Plugins/Online/OnlineSubsystemUtils/Source/OnlineSubsystemUtils/Private/Tests/TestSharingInterface.cpp b/Engine/Plugins/Online/OnlineSubsystemUtils/Source/OnlineSubsystemUtils/Private/Tests/TestSharingInterface.cpp
new file mode 100644
--- /dev/null
+++ b/Engine/Plugins/Online/OnlineSubsystemUtils/Source/OnlineSubsystemUtils/Private/Tests/TestSharingInterface.cpp
@@ -0,0 +1,195 @@
+// Copyright 1998-2016 Epic Games, Inc. All Rights Reserved.
+
+
+#include "OnlineSubsystemUtilsPrivatePCH.h"
+#include "OnlineSharingInterface.h"
+#include "ModuleManager.h"
+#include "ImageCore.h"
+#include "TestSharingInterface.h"
+
+#if WITH_DEV_AUTOMATION_TESTS
+
+FTestSharingInterface::FTestSharingInterface(const FString& InSubsystem)
+{
+	UE_LOG(LogOnline, Verbose, TEXT("FTestSharingInterface::FTestSharingInterface"));
+	SubsystemName = InSubsystem;
+}
+
+
+FTestSharingInterface::~FTestSharingInterface()
+{
+	UE_LOG(LogOnline, Verbose, TEXT("FTestSharingInterface::~FTestSharingInterface"));
+	
+	if(TestStatusUpdate.Image)
+	{
+		delete TestStatusUpdate.Image;
+		TestStatusUpdate.Image = NULL;
+	}
+}
+
+
+void FTestSharingInterface::Test(UWorld* InWorld, bool bWithImage)
+{
+	UE_LOG(LogOnline, Verbose, TEXT("FTestSharingInterface::Test"));
+
+	IOnlineSubsystem* OnlineSub = Online::GetSubsystem(InWorld, FName(*SubsystemName));
+	check(OnlineSub); 
+
+	SharingInterface = OnlineSub->GetSharingInterface();
+	check(SharingInterface.IsValid());
+
+	TestStatusUpdate.Message = FString::Printf(TEXT("This is a test post for UE4 Sharing support! Date = %s"), *FDateTime::Now().ToString());
+	TestStatusUpdate.PostPrivacy = EOnlineStatusUpdatePrivacy::OnlyMe;
+	if( bWithImage )
+	{
+		TestStatusUpdate.Image = new FImage( 256, 256, ERawImageFormat::BGRA8, EGammaSpace::Linear );
+	}
+
+	// Kick off the first part of the test,
+	RequestPermissionsToSharePosts();
+}
+
+
+void FTestSharingInterface::RequestPermissionsToSharePosts()
+{
+	UE_LOG(LogOnline, Verbose, TEXT("FTestSharingInterface::RequestPermissionsToSharePosts"));
+
+	ResponsesReceived = 0;
+	RequestPermissionsToPostToFeedDelegate = FOnRequestNewPublishPermissionsCompleteDelegate::CreateRaw(this, &FTestSharingInterface::OnStatusPostingPermissionsUpdated);
+
+	// We need to be permitted to post on the users behalf to share this update
+	EOnlineSharingPublishingCategory::Type PublishPermissions = EOnlineSharingPublishingCategory::Posts;
+
+	for (int32 PlayerIndex = 0; PlayerIndex < MAX_LOCAL_PLAYERS; PlayerIndex++)
+	{
+		RequestPermissionsToPostToFeedDelegateHandles.Add(PlayerIndex, SharingInterface->AddOnRequestNewPublishPermissionsCompleteDelegate_Handle(PlayerIndex, RequestPermissionsToPostToFeedDelegate));
+		SharingInterface->RequestNewPublishPermissions(PlayerIndex, PublishPermissions, TestStatusUpdate.PostPrivacy);
+	}
+}
+
+
+void FTestSharingInterface::OnStatusPostingPermissionsUpdated(int32 LocalUserNum, bool bWasSuccessful)
+{
+	UE_LOG(LogOnline, Display, TEXT("FTestSharingInterface::OnStatusPostingPermissionsUpdated() - %d"), bWasSuccessful);
+
+	FDelegateHandle DelegateHandle = RequestPermissionsToPostToFeedDelegateHandles.FindRef(LocalUserNum);
+	SharingInterface->ClearOnRequestNewPublishPermissionsCompleteDelegate_Handle(LocalUserNum, DelegateHandle);
+	RequestPermissionsToPostToFeedDelegateHandles.Remove(LocalUserNum);
+
+	if( ++ResponsesReceived == MAX_LOCAL_PLAYERS )
+	{
+		SharePost();
+	}
+}
+
+
+void FTestSharingInterface::SharePost()
+{
+	UE_LOG(LogOnline, Verbose, TEXT("FTestSharingInterface::SharePost"));
+
+	ResponsesReceived = 0;
+	OnPostSharedDelegate = FOnSharePostCompleteDelegate::CreateRaw(this, &FTestSharingInterface::OnPostShared);
+
+	for (int32 PlayerIndex = 0; PlayerIndex < MAX_LOCAL_PLAYERS; PlayerIndex++)
+	{
+		OnPostSharedDelegateHandles.Add(PlayerIndex, SharingInterface->AddOnSharePostCompleteDelegate_Handle(PlayerIndex, OnPostSharedDelegate));
+		SharingInterface->ShareStatusUpdate(PlayerIndex, TestStatusUpdate);
+	}
+}
+
+
+void FTestSharingInterface::OnPostShared(int32 LocalPlayer, bool bWasSuccessful)
+{
+	UE_LOG(LogOnline, Verbose, TEXT("FTestSharingInterface::OnPostShared[PlayerIdx:%i - Successful:%i]"), LocalPlayer, bWasSuccessful);
+
+	FDelegateHandle DelegateHandle = OnPostSharedDelegateHandles.FindRef(LocalPlayer);
+	SharingInterface->ClearOnSharePostCompleteDelegate_Handle(LocalPlayer, DelegateHandle);
+	OnPostSharedDelegateHandles.Remove(LocalPlayer);
+	if( ++ResponsesReceived == MAX_LOCAL_PLAYERS )
+	{
+		RequestPermissionsToReadNewsFeed();
+	}
+}
+
+
+void FTestSharingInterface::RequestPermissionsToReadNewsFeed()
+{
+	UE_LOG(LogOnline, Verbose, TEXT("FTestSharingInterface::RequestPermissionsToReadNewsFeed"));
+
+	ResponsesReceived = 0;
+	RequestPermissionsToReadFeedDelegate = FOnRequestNewReadPermissionsCompleteDelegate::CreateRaw(this, &FTestSharingInterface::OnReadFeedPermissionsUpdated);
+
+	// We need to be permitted to post on the users behalf to share this update
+	EOnlineSharingReadCategory::Type ReadPermissions = EOnlineSharingReadCategory::Posts;
+
+	for (int32 PlayerIndex = 0; PlayerIndex < MAX_LOCAL_PLAYERS; PlayerIndex++)
+	{
+		RequestPermissionsToReadFeedDelegateHandles.Add(PlayerIndex, SharingInterface->AddOnRequestNewReadPermissionsCompleteDelegate_Handle(PlayerIndex, RequestPermissionsToReadFeedDelegate));
+		SharingInterface->RequestNewReadPermissions(PlayerIndex, ReadPermissions);
+	}
+}
+
+
+void FTestSharingInterface::OnReadFeedPermissionsUpdated(int32 LocalUserNum, bool bWasSuccessful)
+{
+	UE_LOG(LogOnline, Display, TEXT("FTestSharingInterface::OnReadFeedPermissionsUpdated() - %d"), bWasSuccessful);
+
+	FDelegateHandle DelegateHandle = RequestPermissionsToReadFeedDelegateHandles.FindRef(LocalUserNum);
+	SharingInterface->ClearOnRequestNewReadPermissionsCompleteDelegate_Handle(LocalUserNum, DelegateHandle);
+	RequestPermissionsToReadFeedDelegateHandles.Remove(LocalUserNum);
+
+	if( ++ResponsesReceived == MAX_LOCAL_PLAYERS )
+	{
+		ReadNewsFeed();
+	}
+}
+
+
+void FTestSharingInterface::ReadNewsFeed()
+{
+	UE_LOG(LogOnline, Verbose, TEXT("FTestSharingInterface::ReadNewsFeed"));
+
+	ResponsesReceived = 0;
+	OnNewsFeedReadDelegate = FOnReadNewsFeedCompleteDelegate::CreateRaw(this, &FTestSharingInterface::OnNewsFeedRead);
+
+	for (int32 PlayerIndex = 0; PlayerIndex < MAX_LOCAL_PLAYERS; PlayerIndex++)
+	{
+		OnNewsFeedReadDelegateHandles.Add(PlayerIndex, SharingInterface->AddOnReadNewsFeedCompleteDelegate_Handle(PlayerIndex, OnNewsFeedReadDelegate));
+		SharingInterface->ReadNewsFeed(PlayerIndex);
+	}
+}
+
+
+void FTestSharingInterface::OnNewsFeedRead(int32 LocalPlayer, bool bWasSuccessful)
+{
+	UE_LOG(LogOnline, Display, TEXT("FTestSharingInterface::OnNewsFeedRead[PlayerIdx:%i - Successful:%i]"), LocalPlayer, bWasSuccessful);
+
+	if (bWasSuccessful)
+	{
+		// Get the 1st cached post
+		FOnlineStatusUpdate FirstReadStatusUpdate;
+		SharingInterface->GetCachedNewsFeed(LocalPlayer, 0, FirstReadStatusUpdate);
+		UE_LOG(LogOnline, Display, TEXT("FTestSharingInterface first read update: %s"), *FirstReadStatusUpdate.Message);
+
+		// Get all the cached posts
+		TArray<FOnlineStatusUpdate> AllReadStatusUpdates;
+		SharingInterface->GetCachedNewsFeeds(LocalPlayer, AllReadStatusUpdates);
+		UE_LOG(LogOnline, Display, TEXT("FTestSharingInterface number of read updates: %d"), AllReadStatusUpdates.Num());
+
+		for (int Idx = 0; Idx < AllReadStatusUpdates.Num(); ++Idx)
+		{
+			const FOnlineStatusUpdate& StatusUpdate = AllReadStatusUpdates[Idx];
+			UE_LOG(LogOnline, Display, TEXT("FTestSharingInterface status update [%d]: %s"), Idx, *StatusUpdate.Message);
+		}
+	}
+
+	FDelegateHandle DelegateHandle = OnNewsFeedReadDelegateHandles.FindRef(LocalPlayer);
+	SharingInterface->ClearOnReadNewsFeedCompleteDelegate_Handle(LocalPlayer, DelegateHandle);
+	if( ++ResponsesReceived == MAX_LOCAL_PLAYERS )
+	{
+		UE_LOG(LogOnline, Display, TEXT("FTestSharingInterface TESTS COMPLETED"), LocalPlayer, bWasSuccessful);
+		delete this;
+	}
+}
+
+#endif //WITH_DEV_AUTOMATION_TESTS
\ No newline at end of file
diff --git a/Engine/Plugins/Online/OnlineSubsystemUtils/Source/OnlineSubsystemUtils/Private/Tests/TestUserInterface.cpp b/Engine/Plugins/Online/OnlineSubsystemUtils/Source/OnlineSubsystemUtils/Private/Tests/TestUserInterface.cpp
new file mode 100644
--- /dev/null
+++ b/Engine/Plugins/Online/OnlineSubsystemUtils/Source/OnlineSubsystemUtils/Private/Tests/TestUserInterface.cpp
@@ -0,0 +1,129 @@
+// Copyright 1998-2016 Epic Games, Inc. All Rights Reserved.
+
+
+#include "OnlineSubsystemUtilsPrivatePCH.h"
+#include "OnlineIdentityInterface.h"
+#include "OnlineUserInterface.h"
+#include "ModuleManager.h"
+#include "TestUserInterface.h"
+
+#if WITH_DEV_AUTOMATION_TESTS
+
+FTestUserInterface::FTestUserInterface(const FString& InSubsystem)
+	: OnlineSub(NULL)
+	, bQueryUserInfo(true)
+{
+	UE_LOG(LogOnline, Display, TEXT("FTestUserInterface::FTestUserInterface"));
+	SubsystemName = InSubsystem;
+}
+
+
+FTestUserInterface::~FTestUserInterface()
+{
+	UE_LOG(LogOnline, Display, TEXT("FTestUserInterface::~FTestUserInterface"));
+}
+
+void FTestUserInterface::Test(UWorld* InWorld, const TArray<FString>& InUserIds)
+{
+	UE_LOG(LogOnline, Display, TEXT("FTestUserInterface::Test"));
+
+	OnlineSub = Online::GetSubsystem(InWorld, SubsystemName.Len() ? FName(*SubsystemName, FNAME_Find) : NAME_None);
+	if (OnlineSub != NULL &&
+		OnlineSub->GetIdentityInterface().IsValid() &&
+		OnlineSub->GetUserInterface().IsValid())
+	{
+		// Add our delegate for the async call
+		OnQueryUserInfoCompleteDelegate       = FOnQueryUserInfoCompleteDelegate::CreateRaw(this, &FTestUserInterface::OnQueryUserInfoComplete);
+		OnQueryUserInfoCompleteDelegateHandle = OnlineSub->GetUserInterface()->AddOnQueryUserInfoCompleteDelegate_Handle(0, OnQueryUserInfoCompleteDelegate);
+
+		// list of users to query
+		for (int32 Idx=0; Idx < InUserIds.Num(); Idx++)
+		{
+			TSharedPtr<const FUniqueNetId> UserId = OnlineSub->GetIdentityInterface()->CreateUniquePlayerId(InUserIds[Idx]);
+			if (UserId.IsValid())
+			{
+				QueryUserIds.Add(UserId.ToSharedRef());
+			}
+		}
+
+		// Include our own user so we have at least one UserId to query for
+		{
+			TSharedPtr<const FUniqueNetId> UserId = OnlineSub->GetIdentityInterface()->GetUniquePlayerId(0);
+			if (QueryUserIds.Contains(UserId.ToSharedRef()) == false)
+			{
+				QueryUserIds.Add(UserId.ToSharedRef());
+			}
+		}
+
+		// kick off next test
+		StartNextTest();
+	}
+	else
+	{
+		UE_LOG(LogOnline, Warning,
+			TEXT("Failed to get user interface for %s"), *SubsystemName);
+		
+		FinishTest();
+	}
+}
+
+void FTestUserInterface::StartNextTest()
+{
+	if (bQueryUserInfo)
+	{
+		OnlineSub->GetUserInterface()->QueryUserInfo(0, QueryUserIds);
+	}
+	else
+	{
+		FinishTest();
+	}
+}
+
+void FTestUserInterface::FinishTest()
+{
+	if (OnlineSub != NULL &&
+		OnlineSub->GetUserInterface().IsValid())
+	{
+		// Clear delegates for the various async calls
+		OnlineSub->GetUserInterface()->ClearOnQueryUserInfoCompleteDelegate_Handle(0, OnQueryUserInfoCompleteDelegateHandle);
+	}
+	delete this;
+}
+
+void FTestUserInterface::OnQueryUserInfoComplete(int32 LocalPlayer, bool bWasSuccessful, const TArray< TSharedRef<const FUniqueNetId> >& UserIds, const FString& ErrorStr)
+{
+	UE_LOG(LogOnline, Log,
+		TEXT("GetUserInterface() for player (%d) was success=%d"), LocalPlayer, bWasSuccessful);
+
+	if (bWasSuccessful)
+	{
+		for (int32 UserIdx=0; UserIdx < UserIds.Num(); UserIdx++)
+		{
+			TSharedPtr<FOnlineUser> User = OnlineSub->GetUserInterface()->GetUserInfo(LocalPlayer, *UserIds[UserIdx]);
+			if (User.IsValid())
+			{
+				UE_LOG(LogOnline, Log,
+					TEXT("PlayerId=%s found"), *UserIds[UserIdx]->ToDebugString());
+				UE_LOG(LogOnline, Log,
+					TEXT("	DisplayName=%s"), *User->GetDisplayName());
+				UE_LOG(LogOnline, Log,
+					TEXT("	RealName=%s"), *User->GetRealName());
+			}
+			else
+			{
+				UE_LOG(LogOnline, Log,
+					TEXT("PlayerId=%s not found"), *UserIds[UserIdx]->ToDebugString());
+			}
+		}
+	}
+	else
+	{
+		UE_LOG(LogOnline, Error, TEXT("GetUserInterface() failure. Error = %s"), *ErrorStr);
+	}
+
+	// done
+	bQueryUserInfo = false;
+	StartNextTest();
+}
+
+#endif //WITH_DEV_AUTOMATION_TESTS
\ No newline at end of file
diff --git a/Engine/Plugins/Online/OnlineSubsystemUtils/Source/OnlineSubsystemUtils/Private/Tests/TestVoice.cpp b/Engine/Plugins/Online/OnlineSubsystemUtils/Source/OnlineSubsystemUtils/Private/Tests/TestVoice.cpp
new file mode 100644
--- /dev/null
+++ b/Engine/Plugins/Online/OnlineSubsystemUtils/Source/OnlineSubsystemUtils/Private/Tests/TestVoice.cpp
@@ -0,0 +1,221 @@
+// Copyright 1998-2016 Epic Games, Inc. All Rights Reserved.
+
+#include "OnlineSubsystemUtilsPrivatePCH.h"
+#include "TestVoice.h"
+#include "TestVoiceData.h"
+
+#include "Engine.h"
+#include "OnlineSubsystemUtils.h"
+#include "Voice.h"
+#include "Runtime/Engine/Classes/Sound/SoundWaveProcedural.h"
+
+#if WITH_DEV_AUTOMATION_TESTS
+
+FTestVoice::FTestVoice() :
+	VoiceComp(NULL),
+	VoiceCapture(NULL),
+	VoiceEncoder(NULL),
+	VoiceDecoder(NULL),
+	bUseTestSample(false),
+	bZeroInput(false),
+	bUseDecompressed(true),
+	bZeroOutput(false)
+{
+	MaxRawCaptureDataSize = 100 * 1024;
+	MaxCompressedDataSize = 50 * 1024;
+	MaxUncompressedDataSize = MaxRawCaptureDataSize;
+	MaxRemainderSize = 1 * 1024;
+	LastRemainderSize = 0;
+
+	RawCaptureData.AddUninitialized(MaxRawCaptureDataSize);
+	CompressedData.AddUninitialized(MaxCompressedDataSize);
+	UncompressedData.AddUninitialized(MaxUncompressedDataSize);
+	Remainder.AddUninitialized(MaxRemainderSize);
+}
+
+FTestVoice::~FTestVoice()
+{
+	Shutdown();
+}
+
+void FTestVoice::Test()
+{
+	Init();
+}
+
+bool FTestVoice::Init()
+{
+	VoiceCapture = FVoiceModule::Get().CreateVoiceCapture();
+	VoiceEncoder = FVoiceModule::Get().CreateVoiceEncoder();
+	VoiceDecoder = FVoiceModule::Get().CreateVoiceDecoder();
+	
+	if (VoiceCapture.IsValid())
+	{
+		VoiceCapture->Start();
+	}
+	
+	return true;
+}
+
+void FTestVoice::Shutdown()
+{
+	CompressedData.Empty();
+	UncompressedData.Empty();
+
+	if (VoiceCapture.IsValid())
+	{
+		VoiceCapture->Shutdown();
+		VoiceCapture = NULL;
+	}
+
+	VoiceEncoder = NULL;
+	VoiceDecoder = NULL;
+
+	if (VoiceComp != NULL)
+	{
+		VoiceComp->RemoveFromRoot();
+		VoiceComp = NULL;
+	}
+}
+
+void FTestVoice::SetStaticVoiceData(TArray<uint8>& VoiceData, uint32& TotalVoiceBytes)
+{
+	static bool bTimeToQueue = true;
+	static double LastQueueTime = 0;
+	double CurrentTime = FPlatformTime::Seconds();
+
+	if (LastQueueTime > 0 && CurrentTime - LastQueueTime > 2)
+	{
+		bTimeToQueue = true;
+	}
+
+	if (bTimeToQueue)
+	{
+		TotalVoiceBytes = ARRAY_COUNT(RawVoiceTestData);
+
+		VoiceData.Empty(TotalVoiceBytes);
+		VoiceData.AddUninitialized(TotalVoiceBytes);
+
+		FMemory::Memcpy(VoiceData.GetData(), RawVoiceTestData, ARRAY_COUNT(RawVoiceTestData));
+
+		LastQueueTime = CurrentTime;
+		bTimeToQueue = false;
+	}
+}
+
+bool FTestVoice::Tick(float DeltaTime)
+{
+	if (VoiceCapture.IsValid())
+	{
+		if (!IsRunningDedicatedServer() && VoiceComp == NULL)
+		{
+			VoiceComp = CreateVoiceAudioComponent(VOICE_SAMPLE_RATE);
+			VoiceComp->AddToRoot();
+			VoiceComp->Play();
+		}
+
+		if (VoiceComp)
+		{
+			static bool bLastHasData = true;
+			USoundWaveProcedural* SoundStreaming = CastChecked<USoundWaveProcedural>(VoiceComp->Sound);
+			bool bHasData = SoundStreaming->GetAvailableAudioByteCount() != 0;
+
+			if (bHasData != bLastHasData)
+			{
+				//UE_LOG(LogVoice, Log, TEXT("VOIP audio component %s starved!"), bHasData ? TEXT("is not") : TEXT("is"));
+				bLastHasData = bHasData;
+			}
+
+			bool bDoWork = false;
+			uint32 TotalVoiceBytes = 0;
+
+			if (bUseTestSample)
+			{
+				SetStaticVoiceData(RawCaptureData, TotalVoiceBytes);
+				bDoWork = true;
+			}
+			else
+			{
+				uint32 NewVoiceDataBytes = 0;
+				EVoiceCaptureState::Type MicState = VoiceCapture->GetCaptureState(NewVoiceDataBytes);
+				if (MicState == EVoiceCaptureState::Ok && NewVoiceDataBytes > 0)
+				{
+					//UE_LOG(LogVoice, Log, TEXT("Getting data! %d"), NewVoiceDataBytes);
+					TotalVoiceBytes = NewVoiceDataBytes + LastRemainderSize;
+					RawCaptureData.Empty(MaxRawCaptureDataSize);
+					RawCaptureData.AddUninitialized(TotalVoiceBytes);
+
+					if (LastRemainderSize > 0)
+					{
+						FMemory::Memcpy(RawCaptureData.GetData(), Remainder.GetData(), LastRemainderSize);
+					}
+
+					MicState = VoiceCapture->GetVoiceData(RawCaptureData.GetData() + LastRemainderSize, NewVoiceDataBytes, NewVoiceDataBytes);
+					TotalVoiceBytes = NewVoiceDataBytes + LastRemainderSize;
+					bDoWork = MicState == EVoiceCaptureState::Ok;
+				}
+			}
+
+			if (bDoWork && TotalVoiceBytes > 0)
+			{
+				// ZERO INPUT
+				if (bZeroInput)
+				{
+					FMemory::Memzero(RawCaptureData.GetData(), TotalVoiceBytes);
+				}	
+				// ZERO INPUT END
+
+				// COMPRESSION BEGIN
+				uint32 CompressedDataSize = MaxCompressedDataSize;
+				LastRemainderSize = VoiceEncoder->Encode(RawCaptureData.GetData(), TotalVoiceBytes, CompressedData.GetData(), CompressedDataSize);
+
+				if (LastRemainderSize > 0)
+				{
+					FMemory::Memcpy(Remainder.GetData(), RawCaptureData.GetData() + (TotalVoiceBytes - LastRemainderSize), LastRemainderSize);
+				}
+				// COMPRESION END
+
+				// DECOMPRESION BEGIN
+				uint32 UncompressedDataSize = MaxUncompressedDataSize;
+				VoiceDecoder->Decode(CompressedData.GetData(), CompressedDataSize, 
+					UncompressedData.GetData(), UncompressedDataSize);
+				// DECOMPRESSION END
+
+				if (bUseDecompressed)
+				{
+					if (UncompressedDataSize > 0)
+					{
+						//UE_LOG(LogVoice, Log, TEXT("Queueing uncompressed data! %d"), UncompressedDataSize);
+						if (bZeroOutput)
+						{
+							FMemory::Memzero((uint8*)UncompressedData.GetData(), UncompressedDataSize);
+						}
+
+						SoundStreaming->QueueAudio(UncompressedData.GetData(), UncompressedDataSize);
+					}
+				}
+				else
+				{
+					//UE_LOG(LogVoice, Log, TEXT("Queueing raw data! %d"), TotalVoiceBytes - LastRemainderSize);
+					SoundStreaming->QueueAudio(RawCaptureData.GetData(), TotalVoiceBytes - LastRemainderSize);
+				}
+			}
+
+		}
+	}
+
+	return true;
+}
+
+bool FTestVoice::Exec(UWorld* InWorld, const TCHAR* Cmd, FOutputDevice& Ar)
+{
+	bool bWasHandled = false;
+	if (FParse::Command(&Cmd, TEXT("killtestvoice")))
+	{
+		delete this;
+		bWasHandled = true;
+	}
+	return bWasHandled;
+}
+
+#endif //WITH_DEV_AUTOMATION_TESTS
\ No newline at end of file
diff --git a/Engine/Plugins/Online/OnlineSubsystemUtils/Source/OnlineSubsystemUtils/Private/Tests/TestVoice.h b/Engine/Plugins/Online/OnlineSubsystemUtils/Source/OnlineSubsystemUtils/Private/Tests/TestVoice.h
new file mode 100644
--- /dev/null
+++ b/Engine/Plugins/Online/OnlineSubsystemUtils/Source/OnlineSubsystemUtils/Private/Tests/TestVoice.h
@@ -0,0 +1,52 @@
+// Copyright 1998-2016 Epic Games, Inc. All Rights Reserved.
+
+#pragma once
+
+#include "Ticker.h"
+
+#if WITH_DEV_AUTOMATION_TESTS
+
+class FTestVoice : public FTickerObjectBase, public FSelfRegisteringExec
+{
+public:
+
+	FTestVoice();
+	virtual ~FTestVoice();
+
+	// FTickerObjectBase
+	virtual bool Tick(float DeltaTime) override;
+
+	// FSelfRegisteringExec
+	virtual bool Exec(UWorld* InWorld, const TCHAR* Cmd, FOutputDevice& Ar) override;
+
+	void Test();
+
+private:
+
+	class UAudioComponent* VoiceComp;
+	TSharedPtr<class IVoiceCapture> VoiceCapture;
+	TSharedPtr<class IVoiceEncoder> VoiceEncoder;
+	TSharedPtr<class IVoiceDecoder> VoiceDecoder;
+
+	TArray<uint8> RawCaptureData;
+	int32 MaxRawCaptureDataSize;
+	TArray<uint8> CompressedData;
+	int32 MaxCompressedDataSize;
+	TArray<uint8> UncompressedData;
+	int32 MaxUncompressedDataSize;
+	TArray<uint8> Remainder;
+	int32 MaxRemainderSize;
+	int32 LastRemainderSize;
+
+	bool bUseTestSample;
+	bool bZeroInput;
+	bool bUseDecompressed;
+	bool bZeroOutput;
+
+	bool Init();
+	void Shutdown();
+
+	void SetStaticVoiceData(TArray<uint8>& VoiceData, uint32& TotalVoiceBytes);
+};
+
+#endif //WITH_DEV_AUTOMATION_TESTS
\ No newline at end of file
