Copying //Tasks/UE4/Dev-Online-Plugins-v3 to //UE4/Dev-Main (Source: //Tasks/UE4/Dev-Online-Plugins-v3 @ 3047294)

diff --git a/Engine/Plugins/Online/OnlineSubsystemUtils/Source/OnlineSubsystemUtils/Private/VoiceEngineImpl.cpp b/Engine/Plugins/Online/OnlineSubsystemUtils/Source/OnlineSubsystemUtils/Private/VoiceEngineImpl.cpp
new file mode 100644
--- /dev/null
+++ b/Engine/Plugins/Online/OnlineSubsystemUtils/Source/OnlineSubsystemUtils/Private/VoiceEngineImpl.cpp
@@ -0,0 +1,465 @@
+// Copyright 1998-2016 Epic Games, Inc. All Rights Reserved.
+
+#include "OnlineSubsystemUtilsPrivatePCH.h"
+#include "VoiceEngineImpl.h"
+#include "OnlineSubsystem.h"
+#include "VoiceInterfaceImpl.h"
+
+#include "Engine.h"
+#include "SoundDefinitions.h"
+#include "Runtime/Engine/Classes/Sound/SoundWaveProcedural.h"
+#include "OnlineSubsystemUtils.h"
+
+/** Largest size preallocated for compressed data */
+#define MAX_COMPRESSED_VOICE_BUFFER_SIZE 8 * 1024
+/** Largest size preallocated for uncompressed data */
+#define MAX_UNCOMPRESSED_VOICE_BUFFER_SIZE 22 * 1024
+/** Largest size allowed to carry over into next buffer */
+#define MAX_VOICE_REMAINDER_SIZE 1 * 1024
+
+FRemoteTalkerDataImpl::FRemoteTalkerDataImpl() :
+	LastSeen(0.0),
+	AudioComponent(nullptr),
+	VoiceDecoder(nullptr)
+{
+	VoiceDecoder = FVoiceModule::Get().CreateVoiceDecoder();
+	check(VoiceDecoder.IsValid());
+}
+
+FRemoteTalkerDataImpl::~FRemoteTalkerDataImpl()
+{
+	VoiceDecoder = nullptr;
+}
+
+FVoiceEngineImpl::FVoiceEngineImpl(IOnlineSubsystem* InSubsystem) :
+	OnlineSubsystem(InSubsystem),
+	VoiceCapture(nullptr),
+	VoiceEncoder(nullptr),
+	OwningUserIndex(INVALID_INDEX),
+	UncompressedBytesAvailable(0),
+	CompressedBytesAvailable(0),
+	AvailableVoiceResult(EVoiceCaptureState::UnInitialized),
+	bPendingFinalCapture(false),
+	bIsCapturing(false),
+	SerializeHelper(nullptr)
+{
+}
+
+FVoiceEngineImpl::~FVoiceEngineImpl()
+{
+	if (bIsCapturing)
+	{
+		VoiceCapture->Stop();
+	}
+
+	VoiceCapture = nullptr;
+	VoiceEncoder = nullptr;
+
+	delete SerializeHelper;
+}
+
+void FVoiceEngineImpl::VoiceCaptureUpdate() const
+{
+	if (bPendingFinalCapture)
+	{
+		uint32 CompressedSize;
+		const EVoiceCaptureState::Type RecordingState = VoiceCapture->GetCaptureState(CompressedSize);
+
+		// If no data is available, we have finished capture the last (post-StopRecording) half-second of voice data
+		if (RecordingState == EVoiceCaptureState::NotCapturing)
+		{
+			UE_LOG(LogVoiceEncode, Log, TEXT("Internal voice capture complete."));
+
+			bPendingFinalCapture = false;
+
+			// If a new recording session has begun since the call to 'StopRecording', kick that off
+			if (bIsCapturing)
+			{
+				StartRecording();
+			}
+			else
+			{
+				// Marks that recording has successfully stopped
+				StoppedRecording();
+			}
+		}
+	}
+}
+
+void FVoiceEngineImpl::StartRecording() const
+{
+	UE_LOG(LogVoiceEncode, VeryVerbose, TEXT("VOIP StartRecording"));
+	if (VoiceCapture.IsValid())
+	{
+		if (!VoiceCapture->Start())
+		{
+			UE_LOG(LogVoiceEncode, Warning, TEXT("Failed to start voice recording"));
+		}
+	}
+}
+
+void FVoiceEngineImpl::StopRecording() const
+{
+	UE_LOG(LogVoiceEncode, VeryVerbose, TEXT("VOIP StopRecording"));
+	if (VoiceCapture.IsValid())
+	{
+		VoiceCapture->Stop();
+	}
+}
+
+void FVoiceEngineImpl::StoppedRecording() const
+{
+	UE_LOG(LogVoiceEncode, VeryVerbose, TEXT("VOIP StoppedRecording"));
+}
+
+bool FVoiceEngineImpl::Init(int32 MaxLocalTalkers, int32 MaxRemoteTalkers)
+{
+	bool bSuccess = false;
+
+	if (!OnlineSubsystem->IsDedicated())
+	{
+		FVoiceModule& VoiceModule = FVoiceModule::Get();
+		if (VoiceModule.IsVoiceEnabled())
+		{
+			VoiceCapture = VoiceModule.CreateVoiceCapture();
+			VoiceEncoder = VoiceModule.CreateVoiceEncoder();
+
+			bSuccess = VoiceCapture.IsValid() && VoiceEncoder.IsValid();
+			if (bSuccess)
+			{
+				CompressedVoiceBuffer.Empty(MAX_COMPRESSED_VOICE_BUFFER_SIZE);
+				DecompressedVoiceBuffer.Empty(MAX_UNCOMPRESSED_VOICE_BUFFER_SIZE);
+
+				for (int32 TalkerIdx = 0; TalkerIdx < MaxLocalTalkers; TalkerIdx++)
+				{
+					PlayerVoiceData[TalkerIdx].VoiceRemainderSize = 0;
+					PlayerVoiceData[TalkerIdx].VoiceRemainder.Empty(MAX_VOICE_REMAINDER_SIZE);
+				}
+			}
+			else
+			{
+				UE_LOG(LogVoice, Warning, TEXT("Voice capture initialization failed!"));
+			}
+		}
+		else
+		{
+			UE_LOG(LogVoice, Log, TEXT("Voice module disabled by config [Voice].bEnabled"));
+		}
+	}
+
+	return bSuccess;
+}
+
+uint32 FVoiceEngineImpl::StartLocalVoiceProcessing(uint32 LocalUserNum) 
+{
+	uint32 Return = E_FAIL;
+	if (IsOwningUser(LocalUserNum))
+	{
+		if (!bIsCapturing)
+		{
+			// Update the current recording state, if VOIP data was still being read
+			VoiceCaptureUpdate();
+
+			if (!IsRecording())
+			{
+				StartRecording();
+			}
+
+			bIsCapturing = true;
+		}
+
+		Return = S_OK;
+	}
+	else
+	{
+		UE_LOG(LogVoiceEncode, Error, TEXT("StartLocalVoiceProcessing(): Device is currently owned by another user"));
+	}
+
+	return Return;
+}
+
+uint32 FVoiceEngineImpl::StopLocalVoiceProcessing(uint32 LocalUserNum) 
+{
+	uint32 Return = E_FAIL;
+	if (IsOwningUser(LocalUserNum))
+	{
+		if (bIsCapturing)
+		{
+			bIsCapturing = false;
+			bPendingFinalCapture = true;
+
+			// Make a call to begin stopping the current VOIP recording session
+			StopRecording();
+
+			// Now check/update the status of the recording session
+			VoiceCaptureUpdate();
+		}
+
+		Return = S_OK;
+	}
+	else
+	{
+		UE_LOG(LogVoiceEncode, Error, TEXT("StopLocalVoiceProcessing: Ignoring stop request for non-owning user"));
+	}
+
+	return Return;
+}
+
+uint32 FVoiceEngineImpl::UnregisterRemoteTalker(const FUniqueNetId& UniqueId)
+{
+	const FUniqueNetIdString& RemoteTalkerId = (const FUniqueNetIdString&)UniqueId;
+	FRemoteTalkerDataImpl* RemoteData = RemoteTalkerBuffers.Find(RemoteTalkerId);
+	if (RemoteData != nullptr)
+	{
+		// Dump the whole talker
+		if (RemoteData->AudioComponent)
+		{
+			RemoteData->AudioComponent->Stop();
+			RemoteData->AudioComponent = nullptr;
+		}
+		
+		RemoteTalkerBuffers.Remove(RemoteTalkerId);
+	}
+
+	return S_OK;
+}
+
+uint32 FVoiceEngineImpl::GetVoiceDataReadyFlags() const
+{
+	// First check and update the internal state of VOIP recording
+	VoiceCaptureUpdate();
+	if (OwningUserIndex != INVALID_INDEX && IsRecording())
+	{
+		// Check if there is new data available via the Voice API
+		if (AvailableVoiceResult == EVoiceCaptureState::Ok && UncompressedBytesAvailable > 0)
+		{
+			return 1 << OwningUserIndex;
+		}
+	}
+
+	return 0;
+}
+
+uint32 FVoiceEngineImpl::ReadLocalVoiceData(uint32 LocalUserNum, uint8* Data, uint32* Size)
+{
+	check(*Size > 0);
+
+	// Before doing anything, check/update the current recording state
+	VoiceCaptureUpdate();
+
+	// Return data even if not capturing, possibly have data during stopping
+	if (IsOwningUser(LocalUserNum) && IsRecording())
+	{
+		DecompressedVoiceBuffer.Empty(MAX_UNCOMPRESSED_VOICE_BUFFER_SIZE);
+		CompressedVoiceBuffer.Empty(MAX_COMPRESSED_VOICE_BUFFER_SIZE);
+
+		uint32 NewVoiceDataBytes = 0;
+		EVoiceCaptureState::Type VoiceResult = VoiceCapture->GetCaptureState(NewVoiceDataBytes);
+		if (VoiceResult != EVoiceCaptureState::Ok && VoiceResult != EVoiceCaptureState::NoData)
+		{
+			UE_LOG(LogVoiceEncode, Warning, TEXT("ReadLocalVoiceData: GetAvailableVoice failure: VoiceResult: %s"), EVoiceCaptureState::ToString(VoiceResult));
+			return E_FAIL;
+		}
+
+ 		if (NewVoiceDataBytes == 0)
+ 		{
+ 			UE_LOG(LogVoiceEncode, VeryVerbose, TEXT("ReadLocalVoiceData: No Data: VoiceResult: %s"), EVoiceCaptureState::ToString(VoiceResult));
+ 			*Size = 0;
+ 			return S_OK;
+ 		}
+
+		// Make space for new and any previously remaining data
+		uint32 TotalVoiceBytes = NewVoiceDataBytes + PlayerVoiceData[LocalUserNum].VoiceRemainderSize;
+		if (TotalVoiceBytes > MAX_UNCOMPRESSED_VOICE_BUFFER_SIZE)
+		{
+			UE_LOG(LogVoiceEncode, Warning, TEXT("Exceeded uncompressed voice buffer size, clamping"))
+			TotalVoiceBytes = MAX_UNCOMPRESSED_VOICE_BUFFER_SIZE;
+		}
+
+		DecompressedVoiceBuffer.AddUninitialized(TotalVoiceBytes);
+
+		if (PlayerVoiceData[LocalUserNum].VoiceRemainderSize > 0)
+		{
+			FMemory::Memcpy(DecompressedVoiceBuffer.GetData(), PlayerVoiceData[LocalUserNum].VoiceRemainder.GetData(), PlayerVoiceData[LocalUserNum].VoiceRemainderSize);
+		}
+
+		// Get new uncompressed data
+		VoiceResult = VoiceCapture->GetVoiceData(DecompressedVoiceBuffer.GetData() + PlayerVoiceData[LocalUserNum].VoiceRemainderSize,
+			NewVoiceDataBytes, NewVoiceDataBytes);
+
+		TotalVoiceBytes = NewVoiceDataBytes + PlayerVoiceData[LocalUserNum].VoiceRemainderSize;
+
+		if (VoiceResult == EVoiceCaptureState::Ok && TotalVoiceBytes > 0)
+		{
+			CompressedBytesAvailable = MAX_COMPRESSED_VOICE_BUFFER_SIZE;
+			CompressedVoiceBuffer.AddUninitialized(MAX_COMPRESSED_VOICE_BUFFER_SIZE);
+
+			PlayerVoiceData[LocalUserNum].VoiceRemainderSize =
+				VoiceEncoder->Encode(DecompressedVoiceBuffer.GetData(), TotalVoiceBytes, CompressedVoiceBuffer.GetData(), CompressedBytesAvailable);
+
+			// Save off any unencoded remainder
+			if (PlayerVoiceData[LocalUserNum].VoiceRemainderSize > 0)
+			{
+				if (PlayerVoiceData[LocalUserNum].VoiceRemainderSize > MAX_VOICE_REMAINDER_SIZE)
+				{
+					UE_LOG(LogVoiceEncode, Warning, TEXT("Exceeded voice remainder buffer size, clamping"))
+						PlayerVoiceData[LocalUserNum].VoiceRemainderSize = MAX_VOICE_REMAINDER_SIZE;
+				}
+
+				PlayerVoiceData[LocalUserNum].VoiceRemainder.AddUninitialized(MAX_VOICE_REMAINDER_SIZE);
+				FMemory::Memcpy(PlayerVoiceData[LocalUserNum].VoiceRemainder.GetData(), DecompressedVoiceBuffer.GetData() + (TotalVoiceBytes - PlayerVoiceData[LocalUserNum].VoiceRemainderSize), PlayerVoiceData[LocalUserNum].VoiceRemainderSize);
+			}
+
+			static double LastGetVoiceCallTime = 0.0;
+			double CurTime = FPlatformTime::Seconds();
+			double TimeSinceLastCall = LastGetVoiceCallTime > 0 ? (CurTime - LastGetVoiceCallTime) : 0.0;
+			LastGetVoiceCallTime = CurTime;
+
+			UE_LOG(LogVoiceEncode, Log, TEXT("ReadLocalVoiceData: GetVoice: Result: %s, Available: %i, LastCall: %0.3f"), EVoiceCaptureState::ToString(VoiceResult), CompressedBytesAvailable, TimeSinceLastCall);
+			if (CompressedBytesAvailable > 0)
+			{
+				*Size = FMath::Min<int32>(*Size, CompressedBytesAvailable);
+				FMemory::Memcpy(Data, CompressedVoiceBuffer.GetData(), *Size);
+
+				UE_LOG(LogVoiceEncode, VeryVerbose, TEXT("ReadLocalVoiceData: Size: %d"), *Size);
+				return S_OK;
+			}
+			else
+			{
+				*Size = 0;
+				CompressedVoiceBuffer.Empty(MAX_COMPRESSED_VOICE_BUFFER_SIZE);
+
+				UE_LOG(LogVoiceEncode, Warning, TEXT("ReadLocalVoiceData: GetVoice failure: VoiceResult: %s"), EVoiceCaptureState::ToString(VoiceResult));
+				return E_FAIL;
+			}
+		}
+	}
+
+	return E_FAIL;
+}
+
+uint32 FVoiceEngineImpl::SubmitRemoteVoiceData(const FUniqueNetId& RemoteTalkerId, uint8* Data, uint32* Size)
+{
+	UE_LOG(LogVoiceDecode, VeryVerbose, TEXT("SubmitRemoteVoiceData(%s) Size: %d received!"), *RemoteTalkerId.ToDebugString(), *Size);
+
+	const FUniqueNetIdString& TalkerId = (const FUniqueNetIdString&)RemoteTalkerId;
+	FRemoteTalkerDataImpl& QueuedData = RemoteTalkerBuffers.FindOrAdd(TalkerId);
+
+	// new voice packet.
+	QueuedData.LastSeen = FPlatformTime::Seconds();
+
+	uint32 BytesWritten = MAX_UNCOMPRESSED_VOICE_BUFFER_SIZE;
+
+	DecompressedVoiceBuffer.Empty(MAX_UNCOMPRESSED_VOICE_BUFFER_SIZE);
+	DecompressedVoiceBuffer.AddUninitialized(MAX_UNCOMPRESSED_VOICE_BUFFER_SIZE);
+	QueuedData.VoiceDecoder->Decode(Data, *Size, DecompressedVoiceBuffer.GetData(), BytesWritten);
+
+	// If there is no data, return
+	if (BytesWritten <= 0)
+	{
+		*Size = 0;
+		return S_OK;
+	}
+
+	bool bAudioComponentCreated = false;
+	// Generate a streaming wave audio component for voice playback
+	if (QueuedData.AudioComponent == nullptr || QueuedData.AudioComponent->IsPendingKill())
+	{
+		if (SerializeHelper == nullptr)
+		{
+			SerializeHelper = new FVoiceSerializeHelper(this);
+		}
+
+		QueuedData.AudioComponent = CreateVoiceAudioComponent(VOICE_SAMPLE_RATE);
+		if (QueuedData.AudioComponent)
+		{
+			QueuedData.AudioComponent->OnAudioFinishedNative.AddRaw(this, &FVoiceEngineImpl::OnAudioFinished);
+		}
+	}
+	
+	if (QueuedData.AudioComponent != nullptr)
+	{
+		if (!QueuedData.AudioComponent->IsActive())
+		{
+			QueuedData.AudioComponent->Play();
+		}
+
+		USoundWaveProcedural* SoundStreaming = CastChecked<USoundWaveProcedural>(QueuedData.AudioComponent->Sound);
+
+		if (0)
+		{
+			if (!bAudioComponentCreated && SoundStreaming->GetAvailableAudioByteCount() == 0)
+			{
+				UE_LOG(LogVoiceDecode, Log, TEXT("VOIP audio component was starved!"));
+			}
+		}
+
+		SoundStreaming->QueueAudio(DecompressedVoiceBuffer.GetData(), BytesWritten);
+	}
+
+	return S_OK;
+}
+
+void FVoiceEngineImpl::TickTalkers(float DeltaTime)
+{
+	// Remove users that are done talking.
+	const double CurTime = FPlatformTime::Seconds();
+	for (FRemoteTalkerData::TIterator It(RemoteTalkerBuffers); It; ++It)
+	{
+		FRemoteTalkerDataImpl& RemoteData = It.Value();
+		double TimeSince = CurTime - RemoteData.LastSeen;
+		if (TimeSince >= 1.0)
+		{
+			// Dump the whole talker
+			if (RemoteData.AudioComponent)
+			{
+				RemoteData.AudioComponent->Stop();
+			}
+		}
+	}
+}
+
+void FVoiceEngineImpl::Tick(float DeltaTime)
+{
+	// Check available voice once a frame, this value changes after calling GetVoiceData()
+	AvailableVoiceResult = VoiceCapture->GetCaptureState(UncompressedBytesAvailable);
+
+	TickTalkers(DeltaTime);
+}
+
+void FVoiceEngineImpl::OnAudioFinished(UAudioComponent* AC)
+{
+	for (FRemoteTalkerData::TIterator It(RemoteTalkerBuffers); It; ++It)
+	{
+		FRemoteTalkerDataImpl& RemoteData = It.Value();
+		if (RemoteData.AudioComponent->IsPendingKill() || AC == RemoteData.AudioComponent)
+		{
+			UE_LOG(LogVoiceDecode, Log, TEXT("Removing VOIP AudioComponent for Id: %s"), *It.Key().ToDebugString());
+			RemoteData.AudioComponent = nullptr;
+			break;
+		}
+	}
+	UE_LOG(LogVoiceDecode, Verbose, TEXT("Audio Finished"));
+}
+
+FString FVoiceEngineImpl::GetVoiceDebugState() const
+{
+	FString Output;
+	Output = FString::Printf(TEXT("IsRecording: %d\n DataReady: 0x%08x State:%s\n UncompressedBytes: %d\n CompressedBytes: %d\n"),
+		IsRecording(), 
+		GetVoiceDataReadyFlags(),
+		EVoiceCaptureState::ToString(AvailableVoiceResult),
+		UncompressedBytesAvailable,
+		CompressedBytesAvailable
+		);
+
+	// Add remainder size
+	for (int32 Idx=0; Idx < MAX_SPLITSCREEN_TALKERS; Idx++)
+	{
+		Output += FString::Printf(TEXT("Remainder[%d] %d\n"), Idx, PlayerVoiceData[Idx].VoiceRemainderSize);
+	}
+
+	return Output;
+}
+
+
diff --git a/Engine/Plugins/Online/OnlineSubsystemUtils/Source/OnlineSubsystemUtils/Public/OnlineSubsystemUtilsModule.h b/Engine/Plugins/Online/OnlineSubsystemUtils/Source/OnlineSubsystemUtils/Public/OnlineSubsystemUtilsModule.h
new file mode 100644
--- /dev/null
+++ b/Engine/Plugins/Online/OnlineSubsystemUtils/Source/OnlineSubsystemUtils/Public/OnlineSubsystemUtilsModule.h
@@ -0,0 +1,44 @@
+// Copyright 1998-2016 Epic Games, Inc. All Rights Reserved.
+
+#pragma once
+
+#include "Core.h"
+#include "ModuleInterface.h"
+
+class IOnlineSubsystemUtils;
+
+/**
+ * Online subsystem utils module class
+ * Misc functionality where dependency on the engine code is allowed (OnlineSubsystem is not allowed to require engine dependencies)
+ */
+class FOnlineSubsystemUtilsModule : public IModuleInterface
+{
+public:
+
+	FOnlineSubsystemUtilsModule() {}
+	virtual ~FOnlineSubsystemUtilsModule() {}
+
+	/** @return the singleton utility interface */
+	IOnlineSubsystemUtils* GetUtils() const { return SubsystemUtils; }
+
+	// IModuleInterface
+	virtual void StartupModule() override;
+	virtual void ShutdownModule() override;
+
+	virtual bool SupportsDynamicReloading() override
+	{
+		return false;
+	}
+
+	virtual bool SupportsAutomaticShutdown() override
+	{
+		return false;
+	}
+
+
+private:
+
+	IOnlineSubsystemUtils* SubsystemUtils;
+};
+
+
diff --git a/Engine/Plugins/Online/OnlineSubsystemUtils/Source/OnlineSubsystemUtils/Public/VoiceEngineImpl.h b/Engine/Plugins/Online/OnlineSubsystemUtils/Source/OnlineSubsystemUtils/Public/VoiceEngineImpl.h
new file mode 100644
--- /dev/null
+++ b/Engine/Plugins/Online/OnlineSubsystemUtils/Source/OnlineSubsystemUtils/Public/VoiceEngineImpl.h
@@ -0,0 +1,254 @@
+// Copyright 1998-2016 Epic Games, Inc. All Rights Reserved.
+
+#pragma once
+
+#include "VoiceInterface.h"
+#include "Voice.h"
+#include "OnlineSubsystemTypes.h"
+#include "OnlineSubsystemUtilsPackage.h"
+#include "Net/VoiceDataCommon.h"
+
+/**
+ * Container for unprocessed voice data
+ */
+struct FLocalVoiceData
+{
+	FLocalVoiceData() :
+		VoiceRemainderSize(0)
+	{
+	}
+
+	/** Amount of voice data not encoded last time */
+	uint32 VoiceRemainderSize;
+	/** Voice sample data not encoded last time */
+	TArray<uint8> VoiceRemainder;
+};
+
+/** 
+ * Remote voice data playing on a single client
+ */
+class FRemoteTalkerDataImpl
+{
+public:
+
+	FRemoteTalkerDataImpl();
+	~FRemoteTalkerDataImpl();
+
+	/** Receive side timestamp since last voice packet fragment */
+	double LastSeen;
+	/** Audio component playing this buffer (only valid on remote instances) */
+	class UAudioComponent* AudioComponent;
+	/** Per remote talker voice decoding state */
+	TSharedPtr<class IVoiceDecoder> VoiceDecoder;
+};
+
+/**
+ * Generic implementation of voice engine, using Voice module for capture/codec
+ */
+class FVoiceEngineImpl : public IVoiceEngine
+{
+	class FVoiceSerializeHelper : public FGCObject
+	{
+		/** Reference to audio components */
+		FVoiceEngineImpl* VoiceEngine;
+		FVoiceSerializeHelper() :
+			VoiceEngine(NULL)
+		{}
+
+	public:
+
+		FVoiceSerializeHelper(FVoiceEngineImpl* InVoiceEngine) :
+			VoiceEngine(InVoiceEngine)
+		{}
+		~FVoiceSerializeHelper() {}
+		
+		/** FGCObject interface */
+		virtual void AddReferencedObjects(FReferenceCollector& Collector) override
+		{
+			// Prevent garbage collection of audio components
+			for (FRemoteTalkerData::TIterator It(VoiceEngine->RemoteTalkerBuffers); It; ++It)
+			{
+				FRemoteTalkerDataImpl& RemoteData = It.Value();
+				if (RemoteData.AudioComponent)
+				{
+					Collector.AddReferencedObject(RemoteData.AudioComponent);
+				}
+			}
+		}
+	};
+
+	friend class FVoiceSerializeHelper;
+
+	/** Mapping of UniqueIds to the incoming voice data and their audio component */
+	typedef TMap<class FUniqueNetIdString, FRemoteTalkerDataImpl> FRemoteTalkerData;
+
+	/** Reference to the main online subsystem */
+	class IOnlineSubsystem* OnlineSubsystem;
+
+	FLocalVoiceData PlayerVoiceData[MAX_SPLITSCREEN_TALKERS];
+	/** Reference to voice capture device */
+	TSharedPtr<class IVoiceCapture> VoiceCapture;
+	/** Reference to voice encoding object */
+	TSharedPtr<class IVoiceEncoder> VoiceEncoder;
+
+    /** User index currently holding onto the voice interface */
+	int32 OwningUserIndex;
+	/** Amount of uncompressed data available this frame */
+	uint32 UncompressedBytesAvailable;
+	/** Amount of compressed data available this frame */
+	uint32 CompressedBytesAvailable;
+	/** Current frame state of voice capture */
+	EVoiceCaptureState::Type AvailableVoiceResult;
+    /** Have we stopped capturing voice but are waiting for its completion */
+	mutable bool bPendingFinalCapture;
+    /** State of voice recording */
+	bool bIsCapturing;
+
+	/** Data from voice codec, waiting to send to network. */
+	TArray<uint8> CompressedVoiceBuffer;
+	/** Data from network playing on an audio component. */
+	FRemoteTalkerData RemoteTalkerBuffers;
+	/** Voice decompression buffer, shared by all talkers */
+	TArray<uint8> DecompressedVoiceBuffer;
+	/** Serialization helper */
+	class FVoiceSerializeHelper* SerializeHelper;
+
+	/**
+	 * Determines if the specified index is the owner or not
+	 *
+	 * @param InIndex the index being tested
+	 *
+	 * @return true if this is the owner, false otherwise
+	 */
+	FORCEINLINE bool IsOwningUser(uint32 UserIndex)
+	{
+		return UserIndex >= 0 && UserIndex < MAX_SPLITSCREEN_TALKERS && OwningUserIndex == UserIndex;
+	}
+
+	/** 
+	 * Update the internal state of the voice capturing state 
+	 * Handles possible continuation waiting for capture stop event
+	 */
+	void VoiceCaptureUpdate() const;
+
+	/** Start capturing voice data */
+	void StartRecording() const;
+
+	/** Stop capturing voice data */
+	void StopRecording() const;
+
+	/** Called when "last half second" is over */
+	void StoppedRecording() const;
+
+	/** @return is active recording occurring at the moment */
+	bool IsRecording() const { return bIsCapturing || bPendingFinalCapture; }
+
+PACKAGE_SCOPE:
+
+	/** Constructor */
+	FVoiceEngineImpl() :
+		OnlineSubsystem(NULL),
+		VoiceCapture(NULL),
+		VoiceEncoder(NULL),
+		OwningUserIndex(INVALID_INDEX),
+		UncompressedBytesAvailable(0),
+		CompressedBytesAvailable(0),
+		AvailableVoiceResult(EVoiceCaptureState::UnInitialized),
+		bPendingFinalCapture(false),
+		bIsCapturing(false),
+		SerializeHelper(NULL)
+	{};
+
+	// IVoiceEngine
+	virtual bool Init(int32 MaxLocalTalkers, int32 MaxRemoteTalkers) override;
+
+public:
+
+	FVoiceEngineImpl(IOnlineSubsystem* InSubsystem);
+	virtual ~FVoiceEngineImpl();
+
+	// IVoiceEngine
+    virtual uint32 StartLocalVoiceProcessing(uint32 LocalUserNum) override;
+    virtual uint32 StopLocalVoiceProcessing(uint32 LocalUserNum) override;
+    
+	virtual uint32 StartRemoteVoiceProcessing(const FUniqueNetId& UniqueId) override
+	{
+		// Not needed
+		return S_OK;
+	}
+
+    virtual uint32 StopRemoteVoiceProcessing(const FUniqueNetId& UniqueId) override
+	{
+		// Not needed
+		return S_OK;
+	}
+
+    virtual uint32 RegisterLocalTalker(uint32 LocalUserNum) override
+	{
+		if (OwningUserIndex == INVALID_INDEX)
+		{
+			OwningUserIndex = LocalUserNum;
+			return S_OK;
+		}
+
+		return E_FAIL;
+	}
+
+    virtual uint32 UnregisterLocalTalker(uint32 LocalUserNum) override
+	{
+		if (IsOwningUser(LocalUserNum))
+		{
+			OwningUserIndex = INVALID_INDEX;
+			return S_OK;
+		}
+
+		return E_FAIL;
+	}
+
+    virtual uint32 RegisterRemoteTalker(const FUniqueNetId& UniqueId) override
+	{
+		// Not needed
+		return S_OK;
+	}
+
+	virtual uint32 UnregisterRemoteTalker(const FUniqueNetId& UniqueId) override;
+
+    virtual bool IsHeadsetPresent(uint32 LocalUserNum) override
+	{
+		return IsOwningUser(LocalUserNum) ? true : false;
+	}
+
+    virtual bool IsLocalPlayerTalking(uint32 LocalUserNum) override
+	{
+		return (GetVoiceDataReadyFlags() & (LocalUserNum << 1)) != 0;
+	}
+
+	virtual bool IsRemotePlayerTalking(const FUniqueNetId& UniqueId) override
+	{
+		return RemoteTalkerBuffers.Find((const FUniqueNetIdString&)UniqueId) != NULL;
+	}
+
+	virtual uint32 GetVoiceDataReadyFlags() const override;
+	virtual uint32 SetPlaybackPriority(uint32 LocalUserNum, const FUniqueNetId& RemoteTalkerId, uint32 Priority) override
+	{
+		// Not supported
+		return S_OK;
+	}
+
+	virtual uint32 ReadLocalVoiceData(uint32 LocalUserNum, uint8* Data, uint32* Size) override;
+	virtual uint32 SubmitRemoteVoiceData(const FUniqueNetId& RemoteTalkerId, uint8* Data, uint32* Size) override;
+	virtual void Tick(float DeltaTime) override;
+	FString GetVoiceDebugState() const override;
+
+	/**
+	 * Update the state of all remote talkers, possibly dropping data or the talker entirely
+	 */
+	void TickTalkers(float DeltaTime);
+
+	/**
+	 * Delegate triggered when an audio component Stop() function is called
+	 */
+	void OnAudioFinished(UAudioComponent* AC);
+};
+
+typedef TSharedPtr<FVoiceEngineImpl, ESPMode::ThreadSafe> FVoiceEngineImplPtr;
diff --git a/Engine/Plugins/Online/OnlineSubsystemUtils/Source/OnlineSubsystemUtils/Public/VoiceInterfaceImpl.h b/Engine/Plugins/Online/OnlineSubsystemUtils/Source/OnlineSubsystemUtils/Public/VoiceInterfaceImpl.h
new file mode 100644
--- /dev/null
+++ b/Engine/Plugins/Online/OnlineSubsystemUtils/Source/OnlineSubsystemUtils/Public/VoiceInterfaceImpl.h
@@ -0,0 +1,137 @@
+// Copyright 1998-2016 Epic Games, Inc. All Rights Reserved.
+
+#pragma once
+
+#include "VoiceInterface.h"
+#include "VoicePacketImpl.h"
+#include "OnlineSubsystemTypes.h"
+#include "OnlineSubsystemUtilsPackage.h"
+
+/**
+ * The generic implementation of the voice interface 
+ */
+
+class ONLINESUBSYSTEMUTILS_API FOnlineVoiceImpl : public IOnlineVoice
+{
+	/** Reference to the main online subsystem */
+	class IOnlineSubsystem* OnlineSubsystem;
+	/** Reference to the sessions interface */
+	class IOnlineSession* SessionInt;
+	/** Reference to the profile interface */
+	class IOnlineIdentity* IdentityInt;
+	/** Reference to the voice engine for acquiring voice data */
+	IVoiceEnginePtr VoiceEngine;
+
+	/** Maximum permitted local talkers */
+	int32 MaxLocalTalkers;
+	/** Maximum permitted remote talkers */
+	int32 MaxRemoteTalkers;
+
+	/** State of all possible local talkers */
+	TArray<FLocalTalker> LocalTalkers;
+	/** State of all possible remote talkers */
+	TArray<FRemoteTalker> RemoteTalkers;
+	/** Remote players locally muted explicitly */
+	TArray<FUniqueNetIdString> SystemMuteList;
+	/** All remote players locally muted */
+	TArray<FUniqueNetIdString> MuteList;
+
+	/** Time to wait for new data before triggering "not talking" */
+	float VoiceNotificationDelta;
+
+	/** Buffered voice data I/O */
+	FVoiceDataImpl VoiceData;
+
+	/**
+	 * Finds a remote talker in the cached list
+	 *
+	 * @param UniqueId the net id of the player to search for
+	 *
+	 * @return pointer to the remote talker or NULL if not found
+	 */
+	struct FRemoteTalker* FindRemoteTalker(const FUniqueNetId& UniqueId);
+
+	/**
+	 * Is a given id presently muted (either by system mute or game server)
+	 */
+	bool IsLocallyMuted(const FUniqueNetId& UniqueId) const;
+
+PACKAGE_SCOPE:
+
+	FOnlineVoiceImpl() :
+		OnlineSubsystem(NULL),
+		SessionInt(NULL),
+		IdentityInt(NULL),
+		VoiceEngine(NULL),
+		MaxLocalTalkers(MAX_SPLITSCREEN_TALKERS),
+		MaxRemoteTalkers(MAX_REMOTE_TALKERS),
+		VoiceNotificationDelta(0.0f)
+	{};
+
+	// IOnlineVoice
+	virtual bool Init() override;
+	void ProcessMuteChangeNotification() override;
+
+	/**
+	 * Cleanup voice interface
+	 */
+	void Shutdown();
+
+	/**
+	 * Processes any talking delegates that need to be fired off
+	 *
+	 * @param DeltaTime the amount of time that has elapsed since the last tick
+	 */
+	void ProcessTalkingDelegates(float DeltaTime);
+
+	/**
+	 * Reads any data that is currently queued
+	 */
+	void ProcessLocalVoicePackets();
+
+	/**
+	 * Submits network packets to audio system for playback
+	 */
+	void ProcessRemoteVoicePackets();
+
+	/**
+	 * Figures out which remote talkers need to be muted for a given local talker
+	 *
+	 * @param TalkerIndex the talker that needs the mute list checked for
+	 * @param PlayerController the player controller associated with this talker
+	 */
+	void UpdateMuteListForLocalTalker(int32 TalkerIndex, class APlayerController* PlayerController);
+
+public:
+
+	/** Constructor */
+	FOnlineVoiceImpl(class IOnlineSubsystem* InOnlineSubsystem);
+
+	/** Virtual destructor to force proper child cleanup */
+	virtual ~FOnlineVoiceImpl();
+
+	// IOnlineVoice
+	virtual void StartNetworkedVoice(uint8 LocalUserNum) override;
+	virtual void StopNetworkedVoice(uint8 LocalUserNum) override;
+    virtual bool RegisterLocalTalker(uint32 LocalUserNum) override;
+	virtual void RegisterLocalTalkers() override;
+    virtual bool UnregisterLocalTalker(uint32 LocalUserNum) override;
+	virtual void UnregisterLocalTalkers() override;
+    virtual bool RegisterRemoteTalker(const FUniqueNetId& UniqueId) override;
+    virtual bool UnregisterRemoteTalker(const FUniqueNetId& UniqueId) override;
+	virtual void RemoveAllRemoteTalkers() override;
+    virtual bool IsHeadsetPresent(uint32 LocalUserNum) override;
+    virtual bool IsLocalPlayerTalking(uint32 LocalUserNum) override;
+	virtual bool IsRemotePlayerTalking(const FUniqueNetId& UniqueId) override;
+	bool IsMuted(uint32 LocalUserNum, const FUniqueNetId& UniqueId) const override;
+	bool MuteRemoteTalker(uint8 LocalUserNum, const FUniqueNetId& PlayerId, bool bIsSystemWide) override;
+	bool UnmuteRemoteTalker(uint8 LocalUserNum, const FUniqueNetId& PlayerId, bool bIsSystemWide) override;
+	virtual TSharedPtr<class FVoicePacket> SerializeRemotePacket(FArchive& Ar) override;
+	virtual TSharedPtr<class FVoicePacket> GetLocalPacket(uint32 LocalUserNum) override;
+	virtual int32 GetNumLocalTalkers() override { return LocalTalkers.Num(); };
+	virtual void ClearVoicePackets() override;
+	virtual void Tick(float DeltaTime) override;
+	virtual FString GetVoiceDebugState() const override;
+};
+
+typedef TSharedPtr<FOnlineVoiceImpl, ESPMode::ThreadSafe> FOnlineVoiceImplPtr;
diff --git a/Engine/Plugins/Online/OnlineSubsystemUtils/Source/OnlineSubsystemUtils/Public/VoicePacketImpl.h b/Engine/Plugins/Online/OnlineSubsystemUtils/Source/OnlineSubsystemUtils/Public/VoicePacketImpl.h
new file mode 100644
--- /dev/null
+++ b/Engine/Plugins/Online/OnlineSubsystemUtils/Source/OnlineSubsystemUtils/Public/VoicePacketImpl.h
@@ -0,0 +1,66 @@
+// Copyright 1998-2016 Epic Games, Inc. All Rights Reserved.
+
+#pragma once
+
+#include "VoiceDataCommon.h"
+#include "OnlineSubsystemUtilsPackage.h"
+
+class FUniqueNetId;
+
+#define DEBUG_VOICE_PACKET_ENCODING 0
+
+/** Defines the data involved in a voice packet */
+class FVoicePacketImpl : public FVoicePacket
+{
+
+PACKAGE_SCOPE:
+
+	/** The unique net id of the talker sending the data */
+	TSharedPtr<const FUniqueNetId> Sender;
+	/** The data that is to be sent/processed */
+	TArray<uint8> Buffer;
+	/** The current amount of space used in the buffer for this packet */
+	uint16 Length;
+
+public:
+	/** Zeros members and validates the assumptions */
+	FVoicePacketImpl() :
+		Sender(NULL),
+		Length(0)
+	{
+		Buffer.Empty(MAX_VOICE_DATA_SIZE);
+		Buffer.AddUninitialized(MAX_VOICE_DATA_SIZE);
+	}
+
+	/** Should only be used by TSharedPtr and FVoiceData */
+	virtual ~FVoicePacketImpl()
+	{
+	}
+
+	/**
+	 * Copies another packet
+	 *
+	 * @param Other packet to copy
+	 */
+	FVoicePacketImpl(const FVoicePacketImpl& Other);
+
+	//~ Begin FVoicePacket interface
+	virtual uint16 GetTotalPacketSize() override;
+	virtual uint16 GetBufferSize() override;
+	virtual TSharedPtr<const FUniqueNetId> GetSender() override;
+	virtual bool IsReliable() override { return false; }
+	virtual void Serialize(class FArchive& Ar) override;
+	//~ End FVoicePacket interface
+};
+
+/** Holds the current voice packet data state */
+struct FVoiceDataImpl
+{
+	/** Data used by the local talkers before sent */
+	FVoicePacketImpl LocalPackets[MAX_SPLITSCREEN_TALKERS];
+	/** Holds the set of received packets that need to be processed */
+	FVoicePacketList RemotePackets;
+
+	FVoiceDataImpl() {}
+	~FVoiceDataImpl() {}
+};
diff --git a/Engine/Plugins/Runtime/Steam/SteamController/Source/SteamController/Private/SteamControllerPrivatePCH.h b/Engine/Plugins/Runtime/Steam/SteamController/Source/SteamController/Private/SteamControllerPrivatePCH.h
--- a/Engine/Plugins/Runtime/Steam/SteamController/Source/SteamController/Private/SteamControllerPrivatePCH.h
+++ b/Engine/Plugins/Runtime/Steam/SteamController/Source/SteamController/Private/SteamControllerPrivatePCH.h
@@ -16,6 +16,9 @@
 #pragma push_macro("ARRAY_COUNT")
 #undef ARRAY_COUNT
 
+#if STEAMSDK_FOUND == 0
+#error Steam SDK not located.  Expected to be found in Engine/Source/ThirdParty/Steamworks/{SteamVersion}
+#endif
 #include <steam/steam_api.h>
 
 #pragma pop_macro("ARRAY_COUNT")
diff --git a/Engine/Plugins/Runtime/Steam/SteamController/Source/SteamController/SteamController.Build.cs b/Engine/Plugins/Runtime/Steam/SteamController/Source/SteamController/SteamController.Build.cs
--- a/Engine/Plugins/Runtime/Steam/SteamController/Source/SteamController/SteamController.Build.cs
+++ b/Engine/Plugins/Runtime/Steam/SteamController/Source/SteamController/SteamController.Build.cs
@@ -1,11 +1,18 @@
 // Copyright 1998-2016 Epic Games, Inc. All Rights Reserved.
 
 using UnrealBuildTool;
+using System.IO;
 
 public class SteamController : ModuleRules
 {
     public SteamController(TargetInfo Target)
     {
+        string SteamVersion = "Steamv132";
+        bool bSteamSDKFound = Directory.Exists(UEBuildConfiguration.UEThirdPartySourceDirectory + "Steamworks/" + SteamVersion) == true;
+
+        Definitions.Add("STEAMSDK_FOUND=" + (bSteamSDKFound ? "1" : "0"));
+        Definitions.Add("WITH_STEAMWORKS=" + (bSteamSDKFound ? "1" : "0"));
+
         PrivateIncludePathModuleNames.Add("TargetPlatform");
 
         PrivateDependencyModuleNames.AddRange(new string[]
@@ -13,9 +20,10 @@ public class SteamController : ModuleRules
 			"Core",
 			"CoreUObject",
 			"Engine",
-			"InputDevice",
 		});
 
+        PublicDependencyModuleNames.Add("InputDevice");
+
         AddEngineThirdPartyPrivateStaticDependencies(Target,
             "Steamworks"
         );
diff --git a/Engine/Source/Editor/LevelEditor/LevelEditor.Build.cs b/Engine/Source/Editor/LevelEditor/LevelEditor.Build.cs
--- a/Engine/Source/Editor/LevelEditor/LevelEditor.Build.cs
+++ b/Engine/Source/Editor/LevelEditor/LevelEditor.Build.cs
@@ -44,7 +44,6 @@ public class LevelEditor : ModuleRules
                 "EditorStyle",
 				"Engine",
 				"MessageLog",
-				"NewsFeed",
                 "SourceControl",
                 "SourceControlWindows",
                 "StatsViewer",
diff --git a/Engine/Source/Editor/LevelEditor/Private/SLevelEditor.cpp b/Engine/Source/Editor/LevelEditor/Private/SLevelEditor.cpp
--- a/Engine/Source/Editor/LevelEditor/Private/SLevelEditor.cpp
+++ b/Engine/Source/Editor/LevelEditor/Private/SLevelEditor.cpp
@@ -31,7 +31,6 @@
 #include "Editor/StatsViewer/Public/StatsViewerModule.h"
 #include "EditorModes.h"
 #include "IDocumentation.h"
-#include "NewsFeed.h"
 #include "TutorialMetaData.h"
 #include "SDockTab.h"
 #include "SActorDetails.h"
@@ -44,7 +43,6 @@
 static const FName LevelEditorBuildAndSubmitTab("LevelEditorBuildAndSubmit");
 static const FName LevelEditorStatsViewerTab("LevelEditorStatsViewer");
 static const FName MainFrameModuleName("MainFrame");
-static const FName NewsFeedModuleName("NewsFeed");
 static const FName LevelEditorModuleName("LevelEditor");
 static const FName WorldBrowserHierarchyTab("WorldBrowserHierarchy");
 static const FName WorldBrowserDetailsTab("WorldBrowserDetails");
@@ -238,21 +236,6 @@ void SLevelEditor::ConstructNotificationBar()
 			FLevelEditorMenu::MakeNotificationBar( LevelEditorCommands, SharedThis(this ) )
 		];
 
-#define SHOW_NEWS_FEED 0
-
-#if SHOW_NEWS_FEED
-	// news feed button
-	INewsFeedModule& NewsFeedModule = FModuleManager::LoadModuleChecked<INewsFeedModule>(NewsFeedModuleName);
-
-	NotificationBarBox->AddSlot()
-		.AutoWidth()
-		.Padding(5.0f, 0.0f, 0.0f, 0.0f)
-		.VAlign(VAlign_Bottom)
-		[
-			NewsFeedModule.CreateNewsFeedButton()
-		];
-
-#endif
 	// developer tools
 	const IMainFrameModule& MainFrameModule = FModuleManager::GetModuleChecked<IMainFrameModule>(MainFrameModuleName);
 
diff --git a/Engine/Source/Editor/UnrealEd/Classes/Editor/EditorEngine.h b/Engine/Source/Editor/UnrealEd/Classes/Editor/EditorEngine.h
--- a/Engine/Source/Editor/UnrealEd/Classes/Editor/EditorEngine.h
+++ b/Engine/Source/Editor/UnrealEd/Classes/Editor/EditorEngine.h
@@ -2490,11 +2490,10 @@ private:
 	 * 
 	 * @param LocalUserNum local user id, for PIE is going to be 0 (there is no splitscreen)
 	 * @param bWasSuccessful was the login successful
-	 * @param UserId userid of the logged in account
 	 * @param ErrorString descriptive error when applicable
 	 * @param DataStruct data required to continue PIE creation, set at login time
 	 */
-	virtual void OnLoginPIEComplete(int32 LocalUserNum, bool bWasSuccessful, const FUniqueNetId& UserId, const FString& ErrorString, FPieLoginStruct DataStruct);
+	virtual void OnLoginPIEComplete(int32 LocalUserNum, bool bWasSuccessful, const FString& ErrorString, FPieLoginStruct DataStruct);
 
 	/** Above function but called a frame later, to stop PIE login from happening from a network callback */
 	virtual void OnLoginPIEComplete_Deferred(int32 LocalUserNum, bool bWasSuccessful, FString ErrorString, FPieLoginStruct DataStruct);
@@ -2846,9 +2845,6 @@ private:
 	*/
 	void MoveViewportCamerasToBox(const FBox& BoundingBox, bool bActiveViewportOnly) const;
 
-	/** Mapping of delegate handles for each online Login() call while in flight */
-	TMap<FName, FDelegateHandle> OnLoginPIECompleteDelegateHandlesForPIEInstances;
-
 public:
 	// Launcher Worker
 	TSharedPtr<class ILauncherWorker> LauncherWorker;
diff --git a/Engine/Source/Editor/UnrealEd/Private/EditorEngine.cpp b/Engine/Source/Editor/UnrealEd/Private/EditorEngine.cpp
--- a/Engine/Source/Editor/UnrealEd/Private/EditorEngine.cpp
+++ b/Engine/Source/Editor/UnrealEd/Private/EditorEngine.cpp
@@ -510,7 +510,7 @@ void UEditorEngine::InitEditor(IEngineLoop* InEngineLoop)
 	UpdateAutoLoadProject();
 
 	// Load any modules that might be required by commandlets
-	FModuleManager::Get().LoadModule(TEXT("OnlineBlueprintSupport"));
+	//FModuleManager::Get().LoadModule(TEXT("OnlineBlueprintSupport"));
 
 	if ( FSlateApplication::IsInitialized() )
 	{
@@ -643,7 +643,7 @@ void UEditorEngine::Init(IEngineLoop* InEngineLoop)
 			TEXT("EditorSettingsViewer"),
 			TEXT("ProjectSettingsViewer"),
 			TEXT("Blutility"),
-			TEXT("OnlineBlueprintSupport"),
+			//TEXT("OnlineBlueprintSupport"),
 			TEXT("XmlParser"),
 			TEXT("UserFeedback"),
 			TEXT("GameplayTagsEditor"),
diff --git a/Engine/Source/Editor/UnrealEd/Private/PlayLevel.cpp b/Engine/Source/Editor/UnrealEd/Private/PlayLevel.cpp
--- a/Engine/Source/Editor/UnrealEd/Private/PlayLevel.cpp
+++ b/Engine/Source/Editor/UnrealEd/Private/PlayLevel.cpp
@@ -34,8 +34,6 @@
 #include "Runtime/Engine/Classes/Engine/UserInterfaceSettings.h"
 #include "Runtime/Engine/Classes/Engine/RendererSettings.h"
 #include "SScissorRectBox.h"
-#include "Online.h"
-#include "OnlineSubsystemUtils.h"
 #include "SNotificationList.h"
 #include "SGameLayerManager.h"
 #include "NotificationManager.h"
@@ -61,6 +59,7 @@
 #include "GameFramework/Pawn.h"
 #include "GameDelegates.h"
 #include "GeneralProjectSettings.h"
+#include "OnlineEngineInterface.h"
 
 #include "AudioThread.h"
 
@@ -254,15 +253,11 @@ void UEditorEngine::EndPlayMap()
 			TeardownPlaySession(ThisContext);
 			
 			// Cleanup online subsystems instantiated during PIE
-			FName OnlineIdentifier = Online::GetUtils()->GetOnlineIdentifier(ThisContext);
-			if (IOnlineSubsystem::DoesInstanceExist(OnlineIdentifier))
+			FName OnlineIdentifier = UOnlineEngineInterface::Get()->GetOnlineIdentifier(ThisContext);
+			if (UOnlineEngineInterface::Get()->DoesInstanceExist(OnlineIdentifier))
 			{
-				IOnlineSubsystem* OnlineSub = IOnlineSubsystem::Get(OnlineIdentifier);
-				if (OnlineSub)
-				{
-					// Stop ticking and clean up, but do not destroy as we may be in a failed online delegate
-					OnlineSub->Shutdown();
-				}
+				// Stop ticking and clean up, but do not destroy as we may be in a failed online delegate
+				UOnlineEngineInterface::Get()->ShutdownOnlineSubsystem(OnlineIdentifier);
 				OnlineIdentifiers.Add(OnlineIdentifier);
 			}
 		
@@ -286,7 +281,6 @@ void UEditorEngine::EndPlayMap()
 		// Clear out viewport index
 		PlayInEditorViewportIndex = -1; 
 
-
 		// We could have been toggling back and forth between simulate and pie before ending the play map
 		// Make sure the property windows are cleared of any pie actors
 		GUnrealEd->UpdateFloatingPropertyWindows();
@@ -468,7 +462,7 @@ void UEditorEngine::CleanupPIEOnlineSessions(TArray<FName> OnlineIdentifiers)
 	for (FName& OnlineIdentifier : OnlineIdentifiers)
 	{
 		UE_LOG(LogPlayLevel, Display, TEXT("Destroying online subsystem %s"), *OnlineIdentifier.ToString());
-		IOnlineSubsystem::Destroy(OnlineIdentifier);
+		UOnlineEngineInterface::Get()->DestroyOnlineSubsystem(OnlineIdentifier);
 		NumOnlinePIEInstances--;
 	}
 
@@ -2339,7 +2333,7 @@ void UEditorEngine::PlayInEditor( UWorld* InWorld, bool bInSimulateInEditor )
 
 	if (SupportsOnlinePIE())
 	{
-		bool bHasRequiredLogins = PlayNumberOfClients <= Online::GetUtils()->GetNumPIELogins();
+		bool bHasRequiredLogins = PlayNumberOfClients <= UOnlineEngineInterface::Get()->GetNumPIELogins();
 		if (bHasRequiredLogins)
 		{
 			// If we support online PIE use it even if we're standalone
@@ -2348,7 +2342,7 @@ void UEditorEngine::PlayInEditor( UWorld* InWorld, bool bInSimulateInEditor )
 		else
 		{
 			FText ErrorMsg = LOCTEXT("PIELoginFailure", "Not enough login credentials to launch all PIE instances, change editor settings");
-			UE_LOG(LogOnline, Verbose, TEXT("%s"), *ErrorMsg.ToString());
+			UE_LOG(LogPlayLevel, Verbose, TEXT("%s"), *ErrorMsg.ToString());
 			FMessageLog(NAME_CategoryPIE).Warning(ErrorMsg);
 		}
 	}
@@ -2569,7 +2563,7 @@ void UEditorEngine::CancelPlayingViaLauncher()
 
 bool UEditorEngine::SupportsOnlinePIE() const
 {
-	return Online::GetUtils()->SupportsOnlinePIE();
+	return UOnlineEngineInterface::Get()->SupportsOnlinePIE();
 }
 
 void UEditorEngine::LoginPIEInstances(bool bAnyBlueprintErrors, bool bStartInSpectatorMode, double PIEStartTime)
@@ -2584,9 +2578,7 @@ void UEditorEngine::LoginPIEInstances(bool bAnyBlueprintErrors, bool bStartInSpe
 	DataStruct.bStartInSpectatorMode = bStartInSpectatorMode;
 	DataStruct.PIEStartTime = PIEStartTime;
 
-	TArray<FOnlineAccountCredentials> PIELogins;
-	Online::GetUtils()->GetPIELogins(PIELogins);
-	ensure(PIELogins.Num() > 0);
+	ensure(UOnlineEngineInterface::Get()->GetNumPIELogins() > 0);
 
 	int32 ClientNum = 0;
 	PIEInstance = 1;
@@ -2613,11 +2605,8 @@ void UEditorEngine::LoginPIEInstances(bool bAnyBlueprintErrors, bool bStartInSpe
 		DataStruct.NetMode = PlayNetMode;
 
 		// Always get the interface (it will create the subsystem regardless)
-		FName OnlineIdentifier = Online::GetUtils()->GetOnlineIdentifier(PieWorldContext);
+		FName OnlineIdentifier = UOnlineEngineInterface::Get()->GetOnlineIdentifier(PieWorldContext);
 		UE_LOG(LogPlayLevel, Display, TEXT("Creating online subsystem for server %s"), *OnlineIdentifier.ToString());
-		IOnlineSubsystem* OnlineSub = IOnlineSubsystem::Get(OnlineIdentifier);
-		IOnlineIdentityPtr IdentityInt = OnlineSub->GetIdentityInterface();
-		check(IdentityInt.IsValid());
 		NumOnlinePIEInstances++;
 
 		if (!CanPlayNetDedicated)
@@ -2627,20 +2616,18 @@ void UEditorEngine::LoginPIEInstances(bool bAnyBlueprintErrors, bool bStartInSpe
 			GetMultipleInstancePositions(DataStruct.SettingsIndex, NextX, NextY);
 
 			// Login to online platform before creating world
-			FOnLoginCompleteDelegate Delegate;
+			FOnPIELoginComplete Delegate;
 			Delegate.BindUObject(this, &UEditorEngine::OnLoginPIEComplete, DataStruct);
 
 			// Login first and continue the flow later
-			FDelegateHandle DelegateHandle = IdentityInt->AddOnLoginCompleteDelegate_Handle(0, Delegate);
-			OnLoginPIECompleteDelegateHandlesForPIEInstances.Add(OnlineIdentifier, DelegateHandle);
-			IdentityInt->Login(0, PIELogins[ClientNum]);
+			UOnlineEngineInterface::Get()->LoginPIEInstance(OnlineIdentifier, 0, ClientNum, Delegate);
 
 			ClientNum++;
 		}
 		else
 		{
 			// Dedicated servers don't use a login
-			OnlineSub->SetForceDedicated(true);
+			UOnlineEngineInterface::Get()->SetForceDedicated(OnlineIdentifier, true);
 			if (CreatePIEWorldFromLogin(PieWorldContext, EPlayNetMode::PIE_ListenServer, DataStruct))
 			{
 				FMessageLog(NAME_CategoryPIE).Info(LOCTEXT("LoggingInDedicated", "Dedicated Server logged in"));
@@ -2669,27 +2656,22 @@ void UEditorEngine::LoginPIEInstances(bool bAnyBlueprintErrors, bool bStartInSpe
 		GetMultipleInstancePositions(DataStruct.SettingsIndex, NextX, NextY);
 		DataStruct.NetMode = WillAutoConnectToServer ? EPlayNetMode::PIE_Client : EPlayNetMode::PIE_Standalone;
 
-		FName OnlineIdentifier = Online::GetUtils()->GetOnlineIdentifier(PieWorldContext);
+		FName OnlineIdentifier = UOnlineEngineInterface::Get()->GetOnlineIdentifier(PieWorldContext);
 		UE_LOG(LogPlayLevel, Display, TEXT("Creating online subsystem for client %s"), *OnlineIdentifier.ToString());
-		IOnlineIdentityPtr IdentityInt = Online::GetIdentityInterface(OnlineIdentifier);
-		check(IdentityInt.IsValid());
 		NumOnlinePIEInstances++;
 
 		// Login to online platform before creating world
-		FOnLoginCompleteDelegate Delegate;
+		FOnPIELoginComplete Delegate;
 		Delegate.BindUObject(this, &UEditorEngine::OnLoginPIEComplete, DataStruct);
 
-		FDelegateHandle DelegateHandle = OnLoginPIECompleteDelegateHandlesForPIEInstances.FindRef(OnlineIdentifier);
-		IdentityInt->ClearOnLoginCompleteDelegate_Handle(0, DelegateHandle);
-		OnLoginPIECompleteDelegateHandlesForPIEInstances.Add(OnlineIdentifier, IdentityInt->AddOnLoginCompleteDelegate_Handle(0, Delegate));
-		IdentityInt->Login(0, PIELogins[ClientNum]);
+		UOnlineEngineInterface::Get()->LoginPIEInstance(OnlineIdentifier, 0, ClientNum, Delegate);
 	}
 
 	// Restore window settings
 	GetMultipleInstancePositions(0, NextX, NextY);	// restore cached settings
 }
 
-void UEditorEngine::OnLoginPIEComplete(int32 LocalUserNum, bool bWasSuccessful, const FUniqueNetId& UserId, const FString& ErrorString, FPieLoginStruct DataStruct)
+void UEditorEngine::OnLoginPIEComplete(int32 LocalUserNum, bool bWasSuccessful, const FString& ErrorString, FPieLoginStruct DataStruct)
 {
 	// This is needed because pie login may change the state of the online objects that called this function
 	GetTimerManager()->SetTimerForNextTick(FTimerDelegate::CreateUObject(this, &UEditorEngine::OnLoginPIEComplete_Deferred, LocalUserNum, bWasSuccessful, ErrorString, DataStruct));
@@ -2697,7 +2679,7 @@ void UEditorEngine::OnLoginPIEComplete(int32 LocalUserNum, bool bWasSuccessful,
 
 void UEditorEngine::OnLoginPIEComplete_Deferred(int32 LocalUserNum, bool bWasSuccessful, FString ErrorString, FPieLoginStruct DataStruct)
 {
-	UE_LOG(LogOnline, Verbose, TEXT("OnLoginPIEComplete LocalUserNum: %d bSuccess: %d %s"), LocalUserNum, bWasSuccessful, *ErrorString);
+	UE_LOG(LogPlayLevel, Verbose, TEXT("OnLoginPIEComplete LocalUserNum: %d bSuccess: %d %s"), LocalUserNum, bWasSuccessful, *ErrorString);
 	FWorldContext* PieWorldContext = GetWorldContextFromHandle(DataStruct.WorldContextHandle);
 
 	if (!PieWorldContext)
@@ -2706,17 +2688,6 @@ void UEditorEngine::OnLoginPIEComplete_Deferred(int32 LocalUserNum, bool bWasSuc
 		return;
 	}
 
-	FName OnlineIdentifier = Online::GetUtils()->GetOnlineIdentifier(*PieWorldContext);
-	IOnlineIdentityPtr IdentityInt = Online::GetIdentityInterface(OnlineIdentifier);
-
-	// Cleanup the login delegate before calling create below
-	FDelegateHandle* DelegateHandle = OnLoginPIECompleteDelegateHandlesForPIEInstances.Find(OnlineIdentifier);
-	if (DelegateHandle)
-	{
-		IdentityInt->ClearOnLoginCompleteDelegate_Handle(0, *DelegateHandle);
-		OnLoginPIECompleteDelegateHandlesForPIEInstances.Remove(OnlineIdentifier);
-	}
-
 	// Create the new world
 	if (CreatePIEWorldFromLogin(*PieWorldContext, DataStruct.NetMode, DataStruct))
 	{
@@ -2787,7 +2758,7 @@ void UEditorEngine::RequestLateJoin()
 
 	if (SupportsOnlinePIE())
 	{
-// 		bool bHasRequiredLogins = PlayNumberOfClients <= Online::GetUtils()->GetNumPIELogins();
+// 		bool bHasRequiredLogins = PlayNumberOfClients <= UOnlineEngineInterface::Get()->GetNumPIELogins();
 // 		if (bHasRequiredLogins)
 		{
 			// If we support online PIE use it even if we're standalone
@@ -2810,9 +2781,6 @@ void UEditorEngine::RequestLateJoin()
 		DataStruct.bStartInSpectatorMode = bStartLateJoinersInSpectatorMode;
 		DataStruct.PIEStartTime = PIEStartTime;
 
-		TArray<FOnlineAccountCredentials> PIELogins;
-		Online::GetUtils()->GetPIELogins(PIELogins);
-
 		FWorldContext &PieWorldContext = CreateNewWorldContext(EWorldType::PIE);
 		PieWorldContext.PIEInstance = PIEInstance++;
 		PieWorldContext.bWaitingOnOnlineSubsystem = true;
@@ -2825,23 +2793,16 @@ void UEditorEngine::RequestLateJoin()
 		GetMultipleInstancePositions(DataStruct.SettingsIndex, NextX, NextY);
 		DataStruct.NetMode = EPlayNetMode::PIE_Client;
 
-		FName OnlineIdentifier = Online::GetUtils()->GetOnlineIdentifier(PieWorldContext);
+		FName OnlineIdentifier = UOnlineEngineInterface::Get()->GetOnlineIdentifier(PieWorldContext);
 		UE_LOG(LogPlayLevel, Display, TEXT("Creating online subsystem for client %s"), *OnlineIdentifier.ToString());
-		IOnlineIdentityPtr IdentityInt = Online::GetIdentityInterface(OnlineIdentifier);
-		check(IdentityInt.IsValid());
 
 		// Login to online platform before creating world
-		FOnLoginCompleteDelegate Delegate;
+		FOnPIELoginComplete Delegate;
 		Delegate.BindUObject(this, &UEditorEngine::OnLoginPIEComplete, DataStruct);
-
-		FDelegateHandle DelegateHandle = OnLoginPIECompleteDelegateHandlesForPIEInstances.FindRef(OnlineIdentifier);
-		IdentityInt->ClearOnLoginCompleteDelegate_Handle(0, DelegateHandle);
-		OnLoginPIECompleteDelegateHandlesForPIEInstances.Add(OnlineIdentifier, IdentityInt->AddOnLoginCompleteDelegate_Handle(0, Delegate));
-		IdentityInt->Login(0, PIELogins[NumOnlinePIEInstances - 1]); // The server doesn't get a PIE login but counts as an online instance
+		// The server doesn't get a PIE login but counts as an online instance
+		UOnlineEngineInterface::Get()->LoginPIEInstance(OnlineIdentifier, 0, NumOnlinePIEInstances - 1, Delegate);
 
 		NumOnlinePIEInstances++;
-
-
 	}
 	else
 	{
diff --git a/Engine/Source/Editor/UnrealEd/UnrealEd.Build.cs b/Engine/Source/Editor/UnrealEd/UnrealEd.Build.cs
--- a/Engine/Source/Editor/UnrealEd/UnrealEd.Build.cs
+++ b/Engine/Source/Editor/UnrealEd/UnrealEd.Build.cs
@@ -97,9 +97,6 @@ public class UnrealEd : ModuleRules
 				"AnimGraph",
                 "AppFramework",
 				"BlueprintGraph",
-				"OnlineSubsystem",
-				"OnlineSubsystemUtils",
-				"OnlineBlueprintSupport",
 				"DesktopPlatform",
 				"EditorStyle",
 				"EngineSettings",
diff --git a/Engine/Source/Programs/CrashReporter/CrashReportClient/CrashReportClient.Target.cs b/Engine/Source/Programs/CrashReporter/CrashReportClient/CrashReportClient.Target.cs
--- a/Engine/Source/Programs/CrashReporter/CrashReportClient/CrashReportClient.Target.cs
+++ b/Engine/Source/Programs/CrashReporter/CrashReportClient/CrashReportClient.Target.cs
@@ -86,7 +86,6 @@ public class CrashReportClientTarget : TargetRules
 		UEBuildConfiguration.bCompileAgainstEngine = false;
 		UEBuildConfiguration.bCompileAgainstCoreUObject = true;
 		UEBuildConfiguration.bUseLoggingInShipping = true;
-		UEBuildConfiguration.bCompileSteamOSS = false;
 
 		UEBuildConfiguration.bIncludeADO = false;
 		
diff --git a/Engine/Source/Programs/UnrealBuildTool/Configuration/UEBuildConfiguration.cs b/Engine/Source/Programs/UnrealBuildTool/Configuration/UEBuildConfiguration.cs
--- a/Engine/Source/Programs/UnrealBuildTool/Configuration/UEBuildConfiguration.cs
+++ b/Engine/Source/Programs/UnrealBuildTool/Configuration/UEBuildConfiguration.cs
@@ -99,18 +99,6 @@ namespace UnrealBuildTool
         [XmlConfig]
         public static bool bCompileSimplygonSSF;
 
-		/// <summary>
-		/// Whether we should compile in support for Steam OnlineSubsystem or not. [RCL] FIXME 2014-Apr-17: bCompileSteamOSS means "bHasSteamworksInstalled" for some code, these meanings need to be untangled
-		/// </summary>
-		[XmlConfig]
-		public static bool bCompileSteamOSS;
-
-		/// <summary>
-		/// Whether we should compile in support for Mcp OnlineSubsystem or not.
-		/// </summary>
-		[XmlConfig]
-		public static bool bCompileMcpOSS;
-
 		/// <summary>
 		/// Whether to compile lean and mean version of UE.
 		/// </summary>
@@ -372,8 +360,6 @@ namespace UnrealBuildTool
 			bUseLoggingInShipping = false;
 			bUseChecksInShipping = false;
             bUseLauncherChecks = false;
-			bCompileSteamOSS = true;
-			bCompileMcpOSS = true;
 			bCompilePhysXVehicle = true;
 			bCompileFreeType = true;
 			bCompileForSize = false;
@@ -395,14 +381,6 @@ namespace UnrealBuildTool
 		public static void PostReset()
 		{
 			// Configuration overrides.
-			string SteamVersion = "Steamv132";
-			bCompileSteamOSS = bCompileSteamOSS
-			   && Directory.Exists(UEBuildConfiguration.UEThirdPartySourceDirectory + "Steamworks/" + SteamVersion) == true;
-
-			bCompileMcpOSS = bCompileMcpOSS
-			   && Directory.Exists("Runtime/Online/NotForLicensees/OnlineSubsystemMcp") == true;
-
-
 			bCompileSimplygon = bCompileSimplygon
 				&& Directory.Exists(UEBuildConfiguration.UEThirdPartySourceDirectory + "NotForLicensees") == true
 				&& Directory.Exists(UEBuildConfiguration.UEThirdPartySourceDirectory + "NotForLicensees/Simplygon") == true
diff --git a/Engine/Source/Runtime/Core/Core.Build.cs b/Engine/Source/Runtime/Core/Core.Build.cs
--- a/Engine/Source/Runtime/Core/Core.Build.cs
+++ b/Engine/Source/Runtime/Core/Core.Build.cs
@@ -157,8 +157,6 @@ public class Core : ModuleRules
         }
         Definitions.Add("UE_ENABLE_ICU=" + (UEBuildConfiguration.bCompileICU ? "1" : "0")); // Enable/disable (=1/=0) ICU usage in the codebase. NOTE: This flag is for use while integrating ICU and will be removed afterward.
 
-        Definitions.Add("WITH_STEAMWORKS=" + (UEBuildConfiguration.bCompileSteamOSS ? "1" : "0"));
-
         // If we're compiling with the engine, then add Core's engine dependencies
 		if (UEBuildConfiguration.bCompileAgainstEngine == true)
 		{
diff --git a/Engine/Source/Runtime/Core/Public/Misc/CallbackDevice.h b/Engine/Source/Runtime/Core/Public/Misc/CallbackDevice.h
--- a/Engine/Source/Runtime/Core/Public/Misc/CallbackDevice.h
+++ b/Engine/Source/Runtime/Core/Public/Misc/CallbackDevice.h
@@ -6,7 +6,6 @@
 
 
 class AActor;
-class FUniqueNetIdString;
 
 
 /**
@@ -82,7 +81,7 @@ public:
 	DECLARE_MULTICAST_DELEGATE(FOnSafeFrameChangedEvent);
 
 	// Callback for handling accepting invitations - generally for engine code
-	DECLARE_MULTICAST_DELEGATE_TwoParams(FOnInviteAccepted, const FString&, const FUniqueNetIdString&);
+	DECLARE_MULTICAST_DELEGATE_TwoParams(FOnInviteAccepted, const FString&, const FString&);
 
 	// Callback for handling the Controller connection / disconnection
 	// first param is true for a connection, false for a disconnection.
diff --git a/Engine/Source/Runtime/CoreUObject/Public/UObject/CoreOnline.h b/Engine/Source/Runtime/CoreUObject/Public/UObject/CoreOnline.h
new file mode 100644
--- /dev/null
+++ b/Engine/Source/Runtime/CoreUObject/Public/UObject/CoreOnline.h
@@ -0,0 +1,258 @@
+// Copyright 1998-2016 Epic Games, Inc. All Rights Reserved.
+
+#pragma once
+
+#include "Core.h"
+#include "CoreUObject.h"
+#include "ObjectBase.h"
+
+const FName GameSessionName(TEXT("Game"));
+const FName PartySessionName(TEXT("Party"));
+
+const FName GamePort(TEXT("GamePort"));
+const FName BeaconPort(TEXT("BeaconPort"));
+
+#if !CPP
+// Circular dependency on Core vs UHT means we have to noexport this struct so tools can build
+USTRUCT(noexport)
+struct FJoinabilitySettings
+{
+	//GENERATED_BODY()
+
+	/** Name of session these settings affect */
+	UPROPERTY()
+	FName SessionName;
+	/** Is this session now publicly searchable */
+	UPROPERTY()
+	bool bPublicSearchable;
+	/** Does this session allow invites */
+	UPROPERTY()
+	bool bAllowInvites;
+	/** Does this session allow public join via presence */
+	UPROPERTY()
+	bool bJoinViaPresence;
+	/** Does this session allow friends to join via presence */
+	UPROPERTY()
+	bool bJoinViaPresenceFriendsOnly;
+	/** Current max players in this session */
+	UPROPERTY()
+	int32 MaxPlayers;
+	/** Current max party size in this session */
+	UPROPERTY()
+	int32 MaxPartySize;
+};
+#endif
+
+struct FJoinabilitySettings
+{
+	/** Name of session these settings affect */
+	FName SessionName;
+	/** Is this session now publicly searchable */
+	bool bPublicSearchable;
+	/** Does this session allow invites */
+	bool bAllowInvites;
+	/** Does this session allow public join via presence */
+	bool bJoinViaPresence;
+	/** Does this session allow friends to join via presence */
+	bool bJoinViaPresenceFriendsOnly;
+	/** Current max players in this session */
+	int32 MaxPlayers;
+	/** Current max party size in this session */
+	int32 MaxPartySize;
+
+	FJoinabilitySettings() :
+		SessionName(NAME_None),
+		bPublicSearchable(false),
+		bAllowInvites(false),
+		bJoinViaPresence(false),
+		bJoinViaPresenceFriendsOnly(false),
+		MaxPlayers(0),
+		MaxPartySize(0)
+	{
+	}
+};
+
+/**
+ * Abstraction of a profile service online Id
+ * The class is meant to be opaque
+ */
+class FUniqueNetId : public TSharedFromThis<FUniqueNetId>
+{
+protected:
+
+	/** Hidden on purpose */
+	FUniqueNetId()
+	{
+	}
+
+	/** Hidden on purpose */
+	FUniqueNetId(const FUniqueNetId& Src)
+	{
+	}
+
+	/** Hidden on purpose */
+	FUniqueNetId& operator=(const FUniqueNetId& Src)
+	{
+		return *this;
+	}
+
+	virtual bool Compare(const FUniqueNetId& Other) const
+	{
+		return (GetSize() == Other.GetSize()) &&
+			(FMemory::Memcmp(GetBytes(), Other.GetBytes(), GetSize()) == 0);
+	}
+
+public:
+
+	/**
+	*	Comparison operator
+	*/
+	bool operator==(const FUniqueNetId& Other) const
+	{
+		return Other.Compare(*this);
+	}
+
+	bool operator!=(const FUniqueNetId& Other) const
+	{
+		return !(FUniqueNetId::operator==(Other));
+	}
+
+	/** 
+	 * Get the raw byte representation of this opaque data
+	 * This data is platform dependent and shouldn't be manipulated directly
+	 *
+	 * @return byte array of size GetSize()
+	 */
+	virtual const uint8* GetBytes() const = 0;
+
+	/** 
+	 * Get the size of the opaque data
+	 *
+	 * @return size in bytes of the data representation
+	 */
+	virtual int32 GetSize() const = 0;
+
+	/**
+	 * Check the validity of the opaque data
+	 *
+	 * @return true if this is well formed data, false otherwise
+	 */
+	virtual bool IsValid() const = 0;
+
+	/**
+	 * Platform specific conversion to string representation of data
+	 *
+	 * @return data in string form
+	 */
+	virtual FString ToString() const = 0;
+
+	/**
+	 * Get a human readable representation of the opaque data
+	 * Shouldn't be used for anything other than logging/debugging
+	 *
+	 * @return data in string form
+	 */
+	virtual FString ToDebugString() const = 0;
+
+	/**
+	 * @return hex encoded string representation of unique id
+	 */
+	FString GetHexEncodedString() const
+	{
+		if (GetSize() > 0 && GetBytes() != NULL)
+		{
+			return BytesToHex(GetBytes(), GetSize());
+		}
+		return FString();
+	}
+
+	virtual ~FUniqueNetId() {}
+};
+
+// This is not a USTRUCT due to cross module dependencies
+struct FUniqueNetIdWrapper
+{
+public:
+
+	FUniqueNetIdWrapper() 
+	: UniqueNetId()
+	{
+	}
+	FUniqueNetIdWrapper(const TSharedPtr<const FUniqueNetId>& InUniqueNetId)
+	: UniqueNetId(InUniqueNetId)
+	{
+	}
+
+	/** Assignment operator */
+	FUniqueNetIdWrapper& operator=(const FUniqueNetIdWrapper& Other)
+	{
+		if (&Other != this)
+		{
+			UniqueNetId = Other.UniqueNetId;
+		}
+		return *this;
+	}
+
+	/** Comparison operator */
+	bool operator==(FUniqueNetIdWrapper const& Other) const
+	{
+		// Both invalid structs or both valid and deep comparison equality
+		bool bBothValid = IsValid() && Other.IsValid();
+		bool bBothInvalid = !IsValid() && !Other.IsValid();
+		return (bBothInvalid || (bBothValid && (*UniqueNetId == *Other.UniqueNetId)));
+	}
+
+	/** Comparison operator */
+	bool operator!=(FUniqueNetIdWrapper const& Other) const
+	{
+		return !(*this == Other);
+	}
+
+	/** Convert this value to a string */
+	FString ToString() const
+	{
+		return IsValid() ? UniqueNetId->ToString() : TEXT("INVALID");
+	}
+
+	/** Is the FUniqueNetId wrapped in this object valid */
+	bool IsValid() const
+	{
+		return UniqueNetId.IsValid() && UniqueNetId->IsValid();
+	}
+	
+	/** 
+	 * Assign a unique id to this wrapper object
+	 *
+	 * @param InUniqueNetId id to associate
+	 */
+	void SetUniqueNetId(const TSharedPtr<const FUniqueNetId>& InUniqueNetId)
+	{
+		UniqueNetId = InUniqueNetId;
+	}
+
+	/** @return unique id associated with this wrapper object */
+	const TSharedPtr<const FUniqueNetId>& GetUniqueNetId() const
+	{
+		return UniqueNetId;
+	}
+
+	/**
+	 * Dereference operator returns a reference to the FUniqueNetId
+	 */
+	const FUniqueNetId& operator*() const
+	{
+		return *UniqueNetId;
+	}
+
+	/**
+	 * Arrow operator returns a pointer to this FUniqueNetId
+	 */
+	const FUniqueNetId* operator->() const
+	{
+		return UniqueNetId.Get();
+	}
+
+protected:
+
+	TSharedPtr<const FUniqueNetId> UniqueNetId;
+};
diff --git a/Engine/Source/Runtime/Engine/Classes/Engine/GameInstance.h b/Engine/Source/Runtime/Engine/Classes/Engine/GameInstance.h
--- a/Engine/Source/Runtime/Engine/Classes/Engine/GameInstance.h
+++ b/Engine/Source/Runtime/Engine/Classes/Engine/GameInstance.h
@@ -57,9 +57,6 @@ protected:
 	UPROPERTY()
 	TArray<ULocalPlayer*> LocalPlayers;		// List of locally participating players in this game instance
 	
-	// Delegate handle that stores delegate for when an invite is accepted by a user
-	FDelegateHandle OnSessionUserInviteAcceptedDelegateHandle;
-	
 	/** Class to manage online services */
 	UPROPERTY()
 	class UOnlineSession* OnlineSession;
@@ -197,12 +194,6 @@ public:
 
 	static void AddReferencedObjects(UObject* InThis, FReferenceCollector& Collector);
 
-    /** Delegate that is called when a user has accepted an invite. */
-	void HandleSessionUserInviteAccepted(const bool bWasSuccess, const int32 ControllerId, TSharedPtr< const FUniqueNetId > UserId, const FOnlineSessionSearchResult & InviteResult);
-	
-	/** Overridable implementation of HandleSessionUserInviteAccepted, which does nothing but call this function */
-	virtual void OnSessionUserInviteAccepted(const bool bWasSuccess, const int32 ControllerId, TSharedPtr< const FUniqueNetId > UserId, const FOnlineSessionSearchResult & InviteResult);
-
 	inline FTimerManager& GetTimerManager() const { return *TimerManager; }
 
 	inline FLatentActionManager& GetLatentActionManager() const { return *LatentActionManager;  }
diff --git a/Engine/Source/Runtime/Engine/Classes/GameFramework/GameSession.h b/Engine/Source/Runtime/Engine/Classes/GameFramework/GameSession.h
--- a/Engine/Source/Runtime/Engine/Classes/GameFramework/GameSession.h
+++ b/Engine/Source/Runtime/Engine/Classes/GameFramework/GameSession.h
@@ -61,7 +61,7 @@ class ENGINE_API AGameSession : public AInfo
 	virtual bool ProcessAutoLogin();
 
     /** Delegate triggered on auto login completion */
-	virtual void OnLoginComplete(int32 LocalUserNum, bool bWasSuccessful, const FUniqueNetId& UserId, const FString& Error);
+	virtual void OnAutoLoginComplete(int32 LocalUserNum, bool bWasSuccessful, const FString& Error);
 
 	/** 
 	 * Called from GameMode.PreLogin() and Login().
@@ -254,12 +254,6 @@ private:
 	 */
 	virtual void OnEndSessionComplete(FName InSessionName, bool bWasSuccessful);
 
-	/** Handle to outstanding start session call */
-	FDelegateHandle StartSessionCompleteHandle;
-	/** Handle to outstanding end session call */
-	FDelegateHandle EndSessionCompleteHandle;
-
-	FDelegateHandle OnLoginCompleteDelegateHandle;
 };
 
 /** 
diff --git a/Engine/Source/Runtime/Engine/Classes/GameFramework/OnlineReplStructs.h b/Engine/Source/Runtime/Engine/Classes/GameFramework/OnlineReplStructs.h
--- a/Engine/Source/Runtime/Engine/Classes/GameFramework/OnlineReplStructs.h
+++ b/Engine/Source/Runtime/Engine/Classes/GameFramework/OnlineReplStructs.h
@@ -5,12 +5,10 @@
 //
 
 #pragma once
-#include "Runtime/Online/OnlineSubsystem/Public/OnlineSubsystemTypes.h"
-#include "Runtime/Online/OnlineSubsystem/Public/UniqueNetIdWrapper.h"
+#include "CoreOnline.h"
 #include "OnlineReplStructs.generated.h"
 
 class FJsonValue;
-class FUniqueNetId;
 
 /**
  * Wrapper for opaque type FUniqueNetId
@@ -27,10 +25,10 @@ struct FUniqueNetIdRepl
 	{
 	}
 
-	FUniqueNetIdRepl(const FUniqueNetIdWrapper& InWrapper)
-	: UniqueNetId(InWrapper.GetUniqueNetId())
-	{
-	}
+ 	FUniqueNetIdRepl(const FUniqueNetIdWrapper& InWrapper)
+ 	: UniqueNetId(InWrapper.GetUniqueNetId())
+ 	{
+ 	}
 
 	FUniqueNetIdRepl(const TSharedRef<const FUniqueNetId>& InUniqueNetId) :
 		UniqueNetId(InUniqueNetId)
@@ -143,10 +141,10 @@ struct FUniqueNetIdRepl
 	/**
 	 * implicit cast operator to FUniqueNetIdWrapper
 	 */
-	FORCEINLINE operator FUniqueNetIdWrapper() const
-	{
-		return FUniqueNetIdWrapper(UniqueNetId);
-	}
+ 	FORCEINLINE operator FUniqueNetIdWrapper() const
+ 	{
+ 		return FUniqueNetIdWrapper(UniqueNetId);
+ 	}
 
 	friend inline uint32 GetTypeHash(FUniqueNetIdRepl const& Value)
 	{
diff --git a/Engine/Source/Runtime/Engine/Engine.Build.cs b/Engine/Source/Runtime/Engine/Engine.Build.cs
--- a/Engine/Source/Runtime/Engine/Engine.Build.cs
+++ b/Engine/Source/Runtime/Engine/Engine.Build.cs
@@ -14,8 +14,6 @@ public class Engine : ModuleRules
 		PrivateIncludePaths.AddRange(
 			new string[] {
 				"Developer/DerivedDataCache/Public",
-				"Runtime/Online/OnlineSubsystem/Public",
-				"Runtime/Online/OnlineSubsystemUtils/Public",
                 "Runtime/SynthBenchmark/Public",
                 "Runtime/Engine/Private",
 			}
@@ -23,7 +21,6 @@ public class Engine : ModuleRules
 
 		PrivateIncludePathModuleNames.AddRange(
 			new string[] {				
-				"OnlineSubsystem",
 				"TargetPlatform",
 				"ImageWrapper",
 				"HeadMountedDisplay",
@@ -266,8 +263,6 @@ public class Engine : ModuleRules
 				"NetworkReplayStreaming",
 				"NullNetworkReplayStreaming",
 				"HttpNetworkReplayStreaming",
-				"OnlineSubsystem", 
-				"OnlineSubsystemUtils",
 				"Advertising"
 			}
 		);
diff --git a/Engine/Source/Runtime/Engine/Private/CheatManager.cpp b/Engine/Source/Runtime/Engine/Private/CheatManager.cpp
--- a/Engine/Source/Runtime/Engine/Private/CheatManager.cpp
+++ b/Engine/Source/Runtime/Engine/Private/CheatManager.cpp
@@ -4,7 +4,7 @@
 #include "UnrealNetwork.h"
 #include "SlateBasics.h"
 #include "NavDataGenerator.h"
-#include "OnlineSubsystemUtils.h"
+#include "Net/OnlineEngineInterface.h"
 #include "VisualLogger/VisualLogger.h"
 #include "AI/Navigation/RecastNavMesh.h"
 #include "GameFramework/Character.h"
@@ -978,29 +978,17 @@ void UCheatManager::SetNavDrawDistance(float DrawDistance)
 
 void UCheatManager::DumpOnlineSessionState()
 {
-	IOnlineSessionPtr SessionInt = Online::GetSessionInterface(GetWorld());
-	if (SessionInt.IsValid())
-	{
-		SessionInt->DumpSessionState();
-	}
+	UOnlineEngineInterface::Get()->DumpSessionState(GetWorld());
 }
 
 void UCheatManager::DumpPartyState()
 {
-	IOnlinePartyPtr PartyInt = Online::GetPartyInterface(GetWorld());
-	if (PartyInt.IsValid())
-	{
-		PartyInt->DumpPartyState();
-	}
+	UOnlineEngineInterface::Get()->DumpPartyState(GetWorld());
 }
 
 void UCheatManager::DumpChatState()
 {
-	IOnlineChatPtr ChatInt = Online::GetChatInterface(GetWorld());
-	if (ChatInt.IsValid())
-	{
-		ChatInt->DumpChatState();
-	}
+	UOnlineEngineInterface::Get()->DumpChatState(GetWorld());
 }
 
 void UCheatManager::DumpVoiceMutingState()
@@ -1010,11 +998,7 @@ void UCheatManager::DumpVoiceMutingState()
 	UE_LOG(LogCheatManager, Display, TEXT(""));
 
 	// Log the online view of the voice state
-	IOnlineVoicePtr VoiceInt = Online::GetVoiceInterface(GetWorld());
-	if (VoiceInt.IsValid())
-	{
-		UE_LOG(LogCheatManager, Display, TEXT("\n%s"), *VoiceInt->GetVoiceDebugState());
-	}
+	UOnlineEngineInterface::Get()->DumpVoiceState(GetWorld());
 
 	// For each player list their gameplay mutes and system wide mutes
 	UE_LOG(LogCheatManager, Display, TEXT("\n%s"), *DumpMutelistState(GetWorld()));
diff --git a/Engine/Source/Runtime/Engine/Private/GameInstance.cpp b/Engine/Source/Runtime/Engine/Private/GameInstance.cpp
--- a/Engine/Source/Runtime/Engine/Private/GameInstance.cpp
+++ b/Engine/Source/Runtime/Engine/Private/GameInstance.cpp
@@ -11,8 +11,6 @@
 #include "Engine/DemoNetDriver.h"
 #include "Engine/LatentActionManager.h"
 #include "Engine/NetworkObjectList.h"
-#include "OnlineSubsystem.h"
-#include "OnlineSessionInterface.h"
 #include "GameFramework/OnlineSession.h"
 #include "GameFramework/PlayerState.h"
 
@@ -61,41 +59,23 @@ void UGameInstance::Init()
 
 	if (!IsRunningCommandlet())
 	{
-		const auto OnlineSub = IOnlineSubsystem::Get();
-		if (OnlineSub != nullptr)
+		UClass* SpawnClass = GetOnlineSessionClass();
+		OnlineSession = NewObject<UOnlineSession>(this, SpawnClass);
+		if (OnlineSession)
 		{
-			IOnlineSessionPtr SessionInt = OnlineSub->GetSessionInterface();
-			if (SessionInt.IsValid())
-			{
-				SessionInt->AddOnSessionUserInviteAcceptedDelegate_Handle(
-					FOnSessionUserInviteAcceptedDelegate::CreateUObject(this, &UGameInstance::HandleSessionUserInviteAccepted));
-			}
+			OnlineSession->RegisterOnlineDelegates();
 		}
 	}
-
-	UClass* SpawnClass = GetOnlineSessionClass();
-	OnlineSession = NewObject<UOnlineSession>(this, SpawnClass);
-	if (OnlineSession)
-	{
-		OnlineSession->RegisterOnlineDelegates();
-	}
 }
 
 void UGameInstance::Shutdown()
 {
 	ReceiveShutdown();
 
-	if (!IsRunningCommandlet())
+	if (OnlineSession)
 	{
-		const auto OnlineSub = IOnlineSubsystem::Get();
-		if (OnlineSub != nullptr)
-		{
-			IOnlineSessionPtr SessionInt = OnlineSub->GetSessionInterface();
-			if (SessionInt.IsValid())
-			{
-				SessionInt->ClearOnSessionUserInviteAcceptedDelegate_Handle(OnSessionUserInviteAcceptedDelegateHandle);
-			}
-		}
+		OnlineSession->ClearOnlineDelegates();
+		OnlineSession = nullptr;
 	}
 
 	for (int32 PlayerIdx = LocalPlayers.Num() - 1; PlayerIdx >= 0; --PlayerIdx)
@@ -901,30 +881,6 @@ TSubclassOf<UOnlineSession> UGameInstance::GetOnlineSessionClass()
 	return UOnlineSession::StaticClass();
 }
 
-
-void UGameInstance::HandleSessionUserInviteAccepted(const bool bWasSuccess, const int32 ControllerId, TSharedPtr< const FUniqueNetId > UserId, const FOnlineSessionSearchResult &	InviteResult)
-{
-	OnSessionUserInviteAccepted(bWasSuccess, ControllerId, UserId, InviteResult);
-}
-
-void UGameInstance::OnSessionUserInviteAccepted(const bool bWasSuccess, const int32 ControllerId, TSharedPtr< const FUniqueNetId > UserId, const FOnlineSessionSearchResult &	InviteResult)
-{
-	UE_LOG(LogPlayerManagement, Verbose, TEXT("OnSessionUserInviteAccepted LocalUserNum: %d bSuccess: %d"), ControllerId, bWasSuccess);
-	// Don't clear invite accept delegate
-
-	if (bWasSuccess)
-	{
-		if (InviteResult.IsValid())
-		{
-			GetOnlineSession()->OnSessionUserInviteAccepted(bWasSuccess, ControllerId, UserId, InviteResult);
-		}
-		else
-		{
-			UE_LOG(LogPlayerManagement, Warning, TEXT("Invite accept returned invalid search result."));
-		}
-	}
-}
-
 bool UGameInstance::IsDedicatedServerInstance() const
 {
 	if (IsRunningDedicatedServer())
diff --git a/Engine/Source/Runtime/Engine/Private/GameMode.cpp b/Engine/Source/Runtime/Engine/Private/GameMode.cpp
--- a/Engine/Source/Runtime/Engine/Private/GameMode.cpp
+++ b/Engine/Source/Runtime/Engine/Private/GameMode.cpp
@@ -8,7 +8,7 @@
 #include "Engine/LevelScriptActor.h"
 #include "GameFramework/GameNetworkManager.h"
 #include "Matinee/MatineeActor.h"
-#include "OnlineSubsystemUtils.h"
+#include "Net/OnlineEngineInterface.h"
 #include "GameFramework/HUD.h"
 #include "GameFramework/DefaultPawn.h"
 #include "GameFramework/SpectatorPawn.h"
@@ -169,15 +169,9 @@ void AGameMode::InitGame(const FString& MapName, const FString& Options, FString
 
 	if (GetNetMode() != NM_Standalone)
 	{
-		FOnlineSessionSettings* SessionSettings = NULL;
-		IOnlineSessionPtr SessionInt = Online::GetSessionInterface(World);
-		if (SessionInt.IsValid())
-		{
-			SessionSettings = SessionInt->GetSessionSettings(GameSession->SessionName);
-		}
-
 		// Attempt to login, returning true means an async login is in flight
-		if (!SessionSettings && !GameSession->ProcessAutoLogin())
+		if (!UOnlineEngineInterface::Get()->DoesSessionExist(World, GameSession->SessionName) &&
+			!GameSession->ProcessAutoLogin())
 		{
 			GameSession->RegisterServer();
 		}
diff --git a/Engine/Source/Runtime/Engine/Private/GameSession.cpp b/Engine/Source/Runtime/Engine/Private/GameSession.cpp
--- a/Engine/Source/Runtime/Engine/Private/GameSession.cpp
+++ b/Engine/Source/Runtime/Engine/Private/GameSession.cpp
@@ -6,7 +6,7 @@
 
 #include "EnginePrivate.h"
 #include "Net/UnrealNetwork.h"
-#include "OnlineSubsystemUtils.h"
+#include "Net/OnlineEngineInterface.h"
 #include "Kismet/GameplayStatics.h"
 #include "GameFramework/PlayerState.h"
 #include "GameFramework/GameSession.h"
@@ -57,8 +57,7 @@ void AGameSession::HandleMatchIsWaitingToStart()
 void AGameSession::HandleMatchHasStarted()
 {
 	UWorld* World = GetWorld();
-	IOnlineSessionPtr SessionInt = Online::GetSessionInterface(World);
-	if (SessionInt.IsValid() && SessionInt->GetNamedSession(SessionName) != nullptr)
+	if (UOnlineEngineInterface::Get()->DoesSessionExist(World, SessionName))
 	{
 		for (FConstPlayerControllerIterator Iterator = World->GetPlayerControllerIterator(); Iterator; ++Iterator)
 		{
@@ -69,8 +68,8 @@ void AGameSession::HandleMatchHasStarted()
 			}
 		}
 
-		StartSessionCompleteHandle = SessionInt->AddOnStartSessionCompleteDelegate_Handle(FOnStartSessionCompleteDelegate::CreateUObject(this, &AGameSession::OnStartSessionComplete));
-		SessionInt->StartSession(SessionName);
+		FOnlineSessionStartComplete CompletionDelegate = FOnlineSessionStartComplete::CreateUObject(this, &AGameSession::OnStartSessionComplete);
+		UOnlineEngineInterface::Get()->StartSession(World, SessionName, CompletionDelegate);
 	}
 
 	if (STATS && !UE_BUILD_SHIPPING)
@@ -86,12 +85,6 @@ void AGameSession::HandleMatchHasStarted()
 void AGameSession::OnStartSessionComplete(FName InSessionName, bool bWasSuccessful)
 {
 	UE_LOG(LogGameSession, Verbose, TEXT("OnStartSessionComplete %s bSuccess: %d"), *InSessionName.ToString(), bWasSuccessful);
-	UWorld* World = GetWorld();
-	IOnlineSessionPtr SessionInt = Online::GetSessionInterface(World);
-	if (SessionInt.IsValid())
-	{
-		SessionInt->ClearOnStartSessionCompleteDelegate_Handle(StartSessionCompleteHandle);
-	}
 }
 
 void AGameSession::HandleMatchHasEnded()
@@ -106,8 +99,7 @@ void AGameSession::HandleMatchHasEnded()
 	}
 
 	UWorld* World = GetWorld();
-	IOnlineSessionPtr SessionInt = Online::GetSessionInterface(World);
-	if (SessionInt.IsValid())
+	if (UOnlineEngineInterface::Get()->DoesSessionExist(World, SessionName))
 	{
 		for (FConstPlayerControllerIterator Iterator = World->GetPlayerControllerIterator(); Iterator; ++Iterator)
 		{
@@ -118,20 +110,14 @@ void AGameSession::HandleMatchHasEnded()
 			}
 		}
 
-		EndSessionCompleteHandle = SessionInt->AddOnEndSessionCompleteDelegate_Handle(FOnEndSessionCompleteDelegate::CreateUObject(this, &AGameSession::OnEndSessionComplete));
-		SessionInt->EndSession(SessionName);
+		FOnlineSessionStartComplete CompletionDelegate = FOnlineSessionEndComplete::CreateUObject(this, &AGameSession::OnEndSessionComplete);
+		UOnlineEngineInterface::Get()->EndSession(World, SessionName, CompletionDelegate);
 	}
 }
 
 void AGameSession::OnEndSessionComplete(FName InSessionName, bool bWasSuccessful)
 {
 	UE_LOG(LogGameSession, Verbose, TEXT("OnEndSessionComplete %s bSuccess: %d"), *InSessionName.ToString(), bWasSuccessful);
-	UWorld* World = GetWorld();
-	IOnlineSessionPtr SessionInt = Online::GetSessionInterface(World);
-	if (SessionInt.IsValid())
-	{
-		SessionInt->ClearOnEndSessionCompleteDelegate_Handle(EndSessionCompleteHandle);
-	}
 }
 
 bool AGameSession::HandleStartMatchRequest()
@@ -165,16 +151,11 @@ void AGameSession::InitOptions( const FString& Options )
 bool AGameSession::ProcessAutoLogin()
 {
 	UWorld* World = GetWorld();
-	IOnlineIdentityPtr IdentityInt = Online::GetIdentityInterface(World);
-	if (IdentityInt.IsValid())
-	{
-		OnLoginCompleteDelegateHandle = IdentityInt->AddOnLoginCompleteDelegate_Handle(0, FOnLoginCompleteDelegate::CreateUObject(this, &AGameSession::OnLoginComplete));
-		if (!IdentityInt->AutoLogin(0))
-		{
-			// Not waiting for async login
-			return false;
-		}
 
+	FOnlineAutoLoginComplete CompletionDelegate = FOnlineAutoLoginComplete::CreateUObject(this, &ThisClass::OnAutoLoginComplete);
+	if (UOnlineEngineInterface::Get()->AutoLogin(World, 0, CompletionDelegate))
+	{
+		// Async login started
 		return true;
 	}
 
@@ -182,21 +163,16 @@ bool AGameSession::ProcessAutoLogin()
 	return false;
 }
 
-void AGameSession::OnLoginComplete(int32 LocalUserNum, bool bWasSuccessful, const FUniqueNetId& UserId, const FString& Error)
+void AGameSession::OnAutoLoginComplete(int32 LocalUserNum, bool bWasSuccessful, const FString& Error)
 {
 	UWorld* World = GetWorld();
-	IOnlineIdentityPtr IdentityInt = Online::GetIdentityInterface(World);
-	if (IdentityInt.IsValid())
+	if (UOnlineEngineInterface::Get()->IsLoggedIn(World, LocalUserNum))
 	{
-		IdentityInt->ClearOnLoginCompleteDelegate_Handle(0, OnLoginCompleteDelegateHandle);
-		if (IdentityInt->GetLoginStatus(0) == ELoginStatus::LoggedIn)
-		{
-			RegisterServer();
-		}
-		else
-		{
-			RegisterServerFailed();
-		}
+		RegisterServer();
+	}
+	else
+	{
+		RegisterServerFailed();
 	}
 }
 
@@ -263,16 +239,12 @@ void AGameSession::RegisterPlayer(APlayerController* NewPlayer, const TSharedPtr
 void AGameSession::UnregisterPlayer(FName InSessionName, const FUniqueNetIdRepl& UniqueId)
 {
 	UWorld* World = GetWorld();
-	IOnlineSessionPtr SessionInt = Online::GetSessionInterface(World);
-	if (SessionInt.IsValid())
+	if (GetNetMode() != NM_Standalone &&
+		UniqueId.IsValid() &&
+		UniqueId->IsValid())
 	{
-		if (GetNetMode() != NM_Standalone &&
-			UniqueId.IsValid() &&
-			UniqueId->IsValid())
-		{
-			// Remove the player from the session
-			SessionInt->UnregisterPlayer(InSessionName, *UniqueId);
-		}
+		// Remove the player from the session
+		UOnlineEngineInterface::Get()->UnregisterPlayer(World, InSessionName, *UniqueId);
 	}
 }
 
@@ -387,25 +359,21 @@ void AGameSession::ReturnToMainMenuHost()
 bool AGameSession::TravelToSession(int32 ControllerId, FName InSessionName)
 {
 	UWorld* World = GetWorld();
-	IOnlineSubsystem* OnlineSub = Online::GetSubsystem(World);
-	if (OnlineSub)
+
+	FString URL;
+	if (UOnlineEngineInterface::Get()->GetResolvedConnectString(World, InSessionName, URL))
 	{
-		FString URL;
-		IOnlineSessionPtr SessionInt = OnlineSub->GetSessionInterface();
-		if (SessionInt.IsValid() && SessionInt->GetResolvedConnectString(InSessionName, URL))
+		APlayerController* PC = UGameplayStatics::GetPlayerController(World, ControllerId);
+		if (PC)
 		{
-			APlayerController* PC = UGameplayStatics::GetPlayerController(World, ControllerId);
-			if (PC)
-			{
-				PC->ClientTravel(URL, TRAVEL_Absolute);
-				return true;
-			}
-		}
-		else
-		{
-			UE_LOG(LogGameSession, Warning, TEXT("Failed to resolve session connect string for %s"), *InSessionName.ToString());
+			PC->ClientTravel(URL, TRAVEL_Absolute);
+			return true;
 		}
 	}
+	else
+	{
+		UE_LOG(LogGameSession, Warning, TEXT("Failed to resolve session connect string for %s"), *InSessionName.ToString());
+	}
 
 	return false;
 }
@@ -419,11 +387,7 @@ void AGameSession::DumpSessionState()
 	UE_LOG(LogGameSession, Log, TEXT("  MaxPlayers: %i"), MaxPlayers);
 	UE_LOG(LogGameSession, Log, TEXT("  MaxSpectators: %i"), MaxSpectators);
 
-	IOnlineSessionPtr SessionInt = Online::GetSessionInterface(GetWorld());
-	if (SessionInt.IsValid())
-	{
-		SessionInt->DumpSessionState();
-	}
+	UOnlineEngineInterface::Get()->DumpSessionState(GetWorld());
 }
 
 bool AGameSession::CanRestartGame()
@@ -436,46 +400,15 @@ bool AGameSession::GetSessionJoinability(FName InSessionName, FJoinabilitySettin
 	UWorld* const World = GetWorld();
 	check(World);
 
-	bool bValidData = false;
-
-	IOnlineSessionPtr SessionInt = Online::GetSessionInterface(World);
-	if (SessionInt.IsValid())
-	{
-		FOnlineSessionSettings* SessionSettings = SessionInt->GetSessionSettings(InSessionName);
-		if (SessionSettings)
-		{
-			OutSettings.SessionName = InSessionName;
-			OutSettings.bPublicSearchable = SessionSettings->bShouldAdvertise;
-			OutSettings.bAllowInvites = SessionSettings->bAllowInvites;
-			OutSettings.bJoinViaPresence = SessionSettings->bAllowJoinViaPresence;
-			OutSettings.bJoinViaPresenceFriendsOnly = SessionSettings->bAllowJoinViaPresenceFriendsOnly;
-
-			OutSettings.MaxPlayers = MaxPlayers;
-			OutSettings.MaxPartySize = MaxPartySize;
-
-			bValidData = true;
-		}
-	}
-
-	return bValidData;
+	OutSettings.MaxPlayers = MaxPlayers;
+	OutSettings.MaxPartySize = MaxPartySize;
+	return UOnlineEngineInterface::Get()->GetSessionJoinability(World, InSessionName, OutSettings);
 }
 
 void AGameSession::UpdateSessionJoinability(FName InSessionName, bool bPublicSearchable, bool bAllowInvites, bool bJoinViaPresence, bool bJoinViaPresenceFriendsOnly)
 {
 	if (GetNetMode() != NM_Standalone)
 	{
-		IOnlineSessionPtr SessionInt = Online::GetSessionInterface(GetWorld());
-		if (SessionInt.IsValid())
-		{
-			FOnlineSessionSettings* GameSettings = SessionInt->GetSessionSettings(InSessionName);
-			if (GameSettings != NULL)
-			{
-				GameSettings->bShouldAdvertise = bPublicSearchable;
-				GameSettings->bAllowInvites = bAllowInvites;
-				GameSettings->bAllowJoinViaPresence = bJoinViaPresence && !bJoinViaPresenceFriendsOnly;
-				GameSettings->bAllowJoinViaPresenceFriendsOnly = bJoinViaPresenceFriendsOnly;
-				SessionInt->UpdateSession(InSessionName, *GameSettings, true);
-			}
-		}
+		UOnlineEngineInterface::Get()->UpdateSessionJoinability(GetWorld(), InSessionName, bPublicSearchable, bAllowInvites, bJoinViaPresence, bJoinViaPresenceFriendsOnly);
 	}
 }
diff --git a/Engine/Source/Runtime/Engine/Private/KismetSystemLibrary.cpp b/Engine/Source/Runtime/Engine/Private/KismetSystemLibrary.cpp
--- a/Engine/Source/Runtime/Engine/Private/KismetSystemLibrary.cpp
+++ b/Engine/Source/Runtime/Engine/Private/KismetSystemLibrary.cpp
@@ -8,15 +8,12 @@
 #include "DelayAction.h"
 #include "InterpolateComponentToAction.h"
 #include "Advertising.h"
-#include "Online.h"
 #include "Camera/CameraActor.h"
 #include "Camera/CameraComponent.h"
 #include "Kismet/GameplayStatics.h"
 #include "SlateCore.h"
 #include "Engine/StreamableManager.h"
-#include "OnlineSubsystemTypes.h"
-#include "OnlineSubsystemUtils.h"
-#include "OnlineIdentityInterface.h"
+#include "Net/OnlineEngineInterface.h"
 #include "UserActivityTracking.h"
 
 //////////////////////////////////////////////////////////////////////////
@@ -3166,40 +3163,32 @@ void UKismetSystemLibrary::ForceCloseAdBanner()
 
 void UKismetSystemLibrary::ShowPlatformSpecificLeaderboardScreen(const FString& CategoryName)
 {
-	IOnlineExternalUIPtr ExternalUI = Online::GetExternalUIInterface();
-	if(ExternalUI.IsValid())
-	{
-		ExternalUI->ShowLeaderboardUI(CategoryName);
-	}
+	// not PIE safe, doesn't have world context
+	UOnlineEngineInterface::Get()->ShowLeaderboardUI(nullptr, CategoryName);
 }
 
-void UKismetSystemLibrary::ShowPlatformSpecificAchievementsScreen(class APlayerController* SpecificPlayer)
+void UKismetSystemLibrary::ShowPlatformSpecificAchievementsScreen(APlayerController* SpecificPlayer)
 {
-	IOnlineExternalUIPtr ExternalUI = Online::GetExternalUIInterface();
-	if(ExternalUI.IsValid())
+	UWorld* World = SpecificPlayer ? SpecificPlayer->GetWorld() : nullptr;
+
+	// Get the controller id from the player
+	int LocalUserNum = 0;
+	if (SpecificPlayer)
 	{
-		// Get the controller id from the player
-		int LocalUserNum = 0;
-		if(SpecificPlayer)
+		ULocalPlayer* LocalPlayer = Cast<ULocalPlayer>(SpecificPlayer->Player);
+		if (LocalPlayer)
 		{
-			ULocalPlayer* LocalPlayer = Cast<ULocalPlayer>(SpecificPlayer->Player);
-			if(LocalPlayer)
-			{
-				LocalUserNum = LocalPlayer->GetControllerId();
-			}
+			LocalUserNum = LocalPlayer->GetControllerId();
 		}
-		ExternalUI->ShowAchievementsUI(LocalUserNum);
 	}
+
+	UOnlineEngineInterface::Get()->ShowAchievementsUI(World, LocalUserNum);
 }
+
 bool UKismetSystemLibrary::IsLoggedIn(APlayerController* SpecificPlayer)
 {
-	IOnlineIdentityPtr Identity = Online::GetIdentityInterface();
-	
-	if (!Identity.IsValid())
-	{
-		return false;
-	}
-	
+	UWorld* World = SpecificPlayer ? SpecificPlayer->GetWorld() : nullptr;
+
 	int LocalUserNum = 0;
 	if (SpecificPlayer != nullptr)
 	{
@@ -3209,7 +3198,8 @@ bool UKismetSystemLibrary::IsLoggedIn(APlayerController* SpecificPlayer)
 			LocalUserNum = LocalPlayer->GetControllerId();
 		}
 	}
-	return Identity->GetLoginStatus(LocalUserNum) == ELoginStatus::LoggedIn;
+
+	return UOnlineEngineInterface::Get()->IsLoggedIn(World, LocalUserNum);
 }
 
 void UKismetSystemLibrary::SetStructurePropertyByName(UObject* Object, FName PropertyName, const FGenericStruct& Value)
diff --git a/Engine/Source/Runtime/Engine/Private/LocalPlayer.cpp b/Engine/Source/Runtime/Engine/Private/LocalPlayer.cpp
--- a/Engine/Source/Runtime/Engine/Private/LocalPlayer.cpp
+++ b/Engine/Source/Runtime/Engine/Private/LocalPlayer.cpp
@@ -7,7 +7,7 @@
 #include "Matinee/InterpGroupInst.h"
 #include "SubtitleManager.h"
 #include "Net/UnrealNetwork.h"
-#include "OnlineSubsystemUtils.h"
+#include "Net/OnlineEngineInterface.h"
 #include "PhysicsPublic.h"
 
 #include "RenderCore.h"
@@ -1446,32 +1446,20 @@ void ULocalPlayer::SetControllerId( int32 NewControllerId )
 
 FString ULocalPlayer::GetNickname() const
 {
-	// Try to get platform identity first
-	IOnlineSubsystem* PlatformSubsystem = IOnlineSubsystem::GetByPlatform(false);
-	if (PlatformSubsystem)
+	UWorld* World = GetWorld();
+	if (World != NULL)
 	{
-		IOnlineIdentityPtr OnlineIdentityInt = PlatformSubsystem->GetIdentityInterface();
-		if (OnlineIdentityInt.IsValid())
+		// Try to get platform identity first
+		FString PlatformNickname;
+		if (UOnlineEngineInterface::Get()->GetPlayerPlatformNickname(World, ControllerId, PlatformNickname))
 		{
-			FString PlayerNickname = OnlineIdentityInt->GetPlayerNickname(ControllerId);
-			if (!PlayerNickname.IsEmpty())
-			{
-				return PlayerNickname;
-			}
+			return PlatformNickname;
 		}
-	}
 
-	UWorld* World = GetWorld();
-	if (World != NULL)
-	{
-		IOnlineIdentityPtr OnlineIdentityInt = Online::GetIdentityInterface(World);
-		if (OnlineIdentityInt.IsValid())
+		auto UniqueId = GetPreferredUniqueNetId();
+		if (UniqueId.IsValid())
 		{
-			auto UniqueId = GetPreferredUniqueNetId();
-			if (UniqueId.IsValid())
-			{
-				return OnlineIdentityInt->GetPlayerNickname(*UniqueId);
-			}
+			return UOnlineEngineInterface::Get()->GetPlayerNickname(World, *UniqueId);
 		}
 	}
 
@@ -1481,20 +1469,12 @@ FString ULocalPlayer::GetNickname() const
 TSharedPtr<const FUniqueNetId> ULocalPlayer::GetUniqueNetIdFromCachedControllerId() const
 {
 	UWorld* World = GetWorld();
-	if (World != NULL)
-	{
-		IOnlineIdentityPtr OnlineIdentityInt = Online::GetIdentityInterface(World);
-		if (OnlineIdentityInt.IsValid())
-		{
-			TSharedPtr<const FUniqueNetId> UniqueId = OnlineIdentityInt->GetUniquePlayerId(ControllerId);
-			if (UniqueId.IsValid())
-			{
-				return UniqueId;
-			}
-		}
+	if (World != nullptr)
+	{		
+		return UOnlineEngineInterface::Get()->GetUniquePlayerId(World, ControllerId);
 	}
 
-	return NULL;
+	return nullptr;
 }
 
 TSharedPtr<const FUniqueNetId> ULocalPlayer::GetCachedUniqueNetId() const
diff --git a/Engine/Source/Runtime/Engine/Private/NetworkDriver.cpp b/Engine/Source/Runtime/Engine/Private/NetworkDriver.cpp
--- a/Engine/Source/Runtime/Engine/Private/NetworkDriver.cpp
+++ b/Engine/Source/Runtime/Engine/Private/NetworkDriver.cpp
@@ -13,7 +13,7 @@
 #include "Engine/VoiceChannel.h"
 #include "Engine/NetworkObjectList.h"
 #include "GameFramework/GameNetworkManager.h"
-#include "OnlineSubsystemUtils.h"
+#include "Net/OnlineEngineInterface.h"
 #include "NetworkingDistanceConstants.h"
 #include "DataChannel.h"
 #include "Engine/PackageMapClient.h"
@@ -818,22 +818,19 @@ void UNetDriver::ProcessLocalServerPackets()
 {
 	if (World)
 	{
-		IOnlineVoicePtr VoiceInt = Online::GetVoiceInterface(World);
-		if (VoiceInt.IsValid())
+		int32 NumLocalTalkers = UOnlineEngineInterface::Get()->GetNumLocalTalkers(World);
+		// Process all of the local packets
+		for (int32 Index = 0; Index < NumLocalTalkers; Index++)
 		{
-			// Process all of the local packets
-			for (int32 Index = 0; Index < VoiceInt->GetNumLocalTalkers(); Index++)
+			// Returns a ref counted copy of the local voice data or NULL if nothing to send
+			TSharedPtr<FVoicePacket> LocalPacket = UOnlineEngineInterface::Get()->GetLocalPacket(World, Index);
+			// Check for something to send for this local talker
+			if (LocalPacket.IsValid())
 			{
-				// Returns a ref counted copy of the local voice data or NULL if nothing to send
-				TSharedPtr<FVoicePacket> LocalPacket = VoiceInt->GetLocalPacket(Index);
-				// Check for something to send for this local talker
-				if (LocalPacket.IsValid())
-				{
-					// See if anyone wants this packet
-					ReplicateVoicePacket(LocalPacket, NULL);
+				// See if anyone wants this packet
+				ReplicateVoicePacket(LocalPacket, NULL);
 
-					// once all local voice packets are processed then call ClearVoicePackets()
-				}
+				// once all local voice packets are processed then call ClearVoicePackets()
 			}
 		}
 	}
@@ -846,17 +843,17 @@ void UNetDriver::ProcessLocalClientPackets()
 {
 	if (World)
 	{
-		IOnlineVoicePtr VoiceInt = Online::GetVoiceInterface(World);
-		if (VoiceInt.IsValid())
+		int32 NumLocalTalkers = UOnlineEngineInterface::Get()->GetNumLocalTalkers(World);
+		if (NumLocalTalkers)
 		{
 			UVoiceChannel* VoiceChannel = ServerConnection->GetVoiceChannel();
 			if (VoiceChannel)
 			{
 				// Process all of the local packets
-				for (int32 Index = 0; Index < VoiceInt->GetNumLocalTalkers(); Index++)
+				for (int32 Index = 0; Index < NumLocalTalkers; Index++)
 				{
 					// Returns a ref counted copy of the local voice data or NULL if nothing to send
-					TSharedPtr<FVoicePacket> LocalPacket = VoiceInt->GetLocalPacket(Index);
+					TSharedPtr<FVoicePacket> LocalPacket = UOnlineEngineInterface::Get()->GetLocalPacket(World, Index);
 					// Check for something to send for this local talker
 					if (LocalPacket.IsValid())
 					{
@@ -879,15 +876,10 @@ void UNetDriver::PostTickFlush()
 {
 	if (World)
 	{
-		IOnlineVoicePtr VoiceInt = Online::GetVoiceInterface(World);
-		if (VoiceInt.IsValid())
-		{
-			VoiceInt->ClearVoicePackets();
-		}
+		UOnlineEngineInterface::Get()->ClearVoicePackets(World);
 	}
 }
 
-
 bool UNetDriver::InitConnectionClass(void)
 {
 	if (NetConnectionClass == NULL && NetConnectionClassName != TEXT(""))
diff --git a/Engine/Source/Runtime/Engine/Private/OnlineBlueprintCallProxyBase.cpp b/Engine/Source/Runtime/Engine/Private/OnlineBlueprintCallProxyBase.cpp
new file mode 100644
--- /dev/null
+++ b/Engine/Source/Runtime/Engine/Private/OnlineBlueprintCallProxyBase.cpp
@@ -0,0 +1,17 @@
+// Copyright 1998-2016 Epic Games, Inc. All Rights Reserved.
+
+#include "EnginePrivate.h"
+#include "OnlineBlueprintCallProxyBase.h"
+
+//////////////////////////////////////////////////////////////////////////
+// UOnlineBlueprintCallProxyBase
+
+UOnlineBlueprintCallProxyBase::UOnlineBlueprintCallProxyBase(const FObjectInitializer& ObjectInitializer)
+	: Super(ObjectInitializer)
+{
+	SetFlags(RF_StrongRefOnFrame);
+}
+
+void UOnlineBlueprintCallProxyBase::Activate()
+{
+}
diff --git a/Engine/Source/Runtime/Engine/Private/OnlineEngineInterface.cpp b/Engine/Source/Runtime/Engine/Private/OnlineEngineInterface.cpp
new file mode 100644
--- /dev/null
+++ b/Engine/Source/Runtime/Engine/Private/OnlineEngineInterface.cpp
@@ -0,0 +1,34 @@
+// Copyright 1998-2016 Epic Games, Inc. All Rights Reserved.
+
+#include "EnginePrivate.h"
+#include "Net/OnlineEngineInterface.h"
+
+PRAGMA_DISABLE_OPTIMIZATION
+
+UOnlineEngineInterface* UOnlineEngineInterface::Singleton = nullptr;
+
+UOnlineEngineInterface::UOnlineEngineInterface(const FObjectInitializer& ObjectInitializer)
+	: Super(ObjectInitializer)
+{
+}
+
+UOnlineEngineInterface* UOnlineEngineInterface::Get()
+{
+	if (!Singleton)
+	{
+		// Proper interface class hard coded here to emphasize the fact that this is not expected to change much, any need to do so should go through the OGS team first
+		UClass* OnlineEngineInterfaceClass = StaticLoadClass(UOnlineEngineInterface::StaticClass(), NULL, TEXT("/Script/OnlineSubsystemUtils.OnlineEngineInterfaceImpl"), NULL, LOAD_Quiet, NULL);
+		if (!OnlineEngineInterfaceClass)
+		{
+			// Default to the no op class if necessary
+			OnlineEngineInterfaceClass = UOnlineEngineInterface::StaticClass();
+		}
+
+		Singleton = NewObject<UOnlineEngineInterface>(GetTransientPackage(), OnlineEngineInterfaceClass);
+		Singleton->AddToRoot();
+	}
+
+	return Singleton;
+}
+
+PRAGMA_ENABLE_OPTIMIZATION
diff --git a/Engine/Source/Runtime/Engine/Private/OnlineReplStructs.cpp b/Engine/Source/Runtime/Engine/Private/OnlineReplStructs.cpp
--- a/Engine/Source/Runtime/Engine/Private/OnlineReplStructs.cpp
+++ b/Engine/Source/Runtime/Engine/Private/OnlineReplStructs.cpp
@@ -5,7 +5,7 @@
 =============================================================================*/
 
 #include "EnginePrivate.h"
-#include "OnlineSubsystemUtils.h"
+#include "Net/OnlineEngineInterface.h"
 
 FArchive& operator<<( FArchive& Ar, FUniqueNetIdRepl& UniqueNetId)
 {
@@ -35,12 +35,8 @@ FArchive& operator<<( FArchive& Ar, FUniqueNetIdRepl& UniqueNetId)
 void FUniqueNetIdRepl::UniqueIdFromString(const FString& Contents)
 {
 	// Don't need to distinguish OSS interfaces here with world because we just want the create function below
-	IOnlineIdentityPtr IdentityInt = Online::GetIdentityInterface();
-	if (IdentityInt.IsValid())
-	{
-		TSharedPtr<const FUniqueNetId> UniqueNetIdPtr = IdentityInt->CreateUniquePlayerId(Contents);
-		SetUniqueNetId(UniqueNetIdPtr);
-	}
+	TSharedPtr<const FUniqueNetId> UniqueNetIdPtr = UOnlineEngineInterface::Get()->CreateUniquePlayerId(Contents);
+	SetUniqueNetId(UniqueNetIdPtr);
 }
 
 bool FUniqueNetIdRepl::NetSerialize(FArchive& Ar, class UPackageMap* Map, bool& bOutSuccess)
@@ -82,7 +78,7 @@ bool FUniqueNetIdRepl::ImportTextItem(const TCHAR*& Buffer, int32 PortFlags, UOb
 		}
 		else
 		{
-			checkf(IOnlineSubsystem::IsLoaded(), TEXT("Attempted to ImportText to FUniqueNetIdRepl while OSS is not loaded. Parent:%s"), *GetPathNameSafe(Parent));
+			checkf(UOnlineEngineInterface::Get() && UOnlineEngineInterface::Get()->IsLoaded(), TEXT("Attempted to ImportText to FUniqueNetIdRepl while OSS is not loaded. Parent:%s"), *GetPathNameSafe(Parent));
 			FString Contents(Buffer);
 			UniqueIdFromString(Contents);
 		}
@@ -125,81 +121,78 @@ void TestUniqueIdRepl(UWorld* InWorld)
 #if !UE_BUILD_SHIPPING
 	bool bSuccess = true;
 
-	IOnlineIdentityPtr IdentityPtr = Online::GetIdentityInterface(InWorld);
-	if (IdentityPtr.IsValid())
+	TSharedPtr<const FUniqueNetId> UserId = UOnlineEngineInterface::Get()->GetUniquePlayerId(InWorld, 0);
+
+	FUniqueNetIdRepl EmptyIdIn;
+	if (EmptyIdIn.IsValid())
 	{
-		TSharedPtr<const FUniqueNetId> UserId = IdentityPtr->GetUniquePlayerId(0);
+		UE_LOG(LogNet, Warning, TEXT("EmptyId is valid."), *EmptyIdIn->ToString());
+		bSuccess = false;
+	}
 
-		FUniqueNetIdRepl EmptyIdIn;
-		if (EmptyIdIn.IsValid())
-		{
-			UE_LOG(LogNet, Warning, TEXT("EmptyId is valid."), *EmptyIdIn->ToString());
-			bSuccess = false;
-		}
+	FUniqueNetIdRepl ValidIdIn(UserId);
+	if (!ValidIdIn.IsValid() || UserId != ValidIdIn.GetUniqueNetId() || *UserId != *ValidIdIn)
+	{
+		UE_LOG(LogNet, Warning, TEXT("UserId input %s != UserId output %s"), *UserId->ToString(), *ValidIdIn->ToString());
+		bSuccess = false;
+	}
 
-		FUniqueNetIdRepl ValidIdIn(UserId);
-		if (!ValidIdIn.IsValid() || UserId != ValidIdIn.GetUniqueNetId() || *UserId != *ValidIdIn)
+	if (bSuccess)
+	{
+		TArray<uint8> Buffer;
+		for (int32 i = 0; i < 2; i++)
 		{
-			UE_LOG(LogNet, Warning, TEXT("UserId input %s != UserId output %s"), *UserId->ToString(), *ValidIdIn->ToString());
-			bSuccess = false;
-		}
+			Buffer.Empty();
+			FMemoryWriter TestWriteUniqueId(Buffer);
 
-		if (bSuccess)
-		{
-			TArray<uint8> Buffer;
-			for (int32 i=0; i<2; i++)
+			if (i == 0)
+			{
+				// Normal serialize
+				TestWriteUniqueId << EmptyIdIn;
+				TestWriteUniqueId << ValidIdIn;
+			}
+			else
 			{
-				Buffer.Empty();
-				FMemoryWriter TestWriteUniqueId(Buffer);
-
-				if (i == 0)
-				{
-					// Normal serialize
-					TestWriteUniqueId << EmptyIdIn;
-					TestWriteUniqueId << ValidIdIn;
-				}
-				else
-				{
-					// Net serialize
-					bool bOutSuccess = false;
-					EmptyIdIn.NetSerialize(TestWriteUniqueId, NULL, bOutSuccess);
-					ValidIdIn.NetSerialize(TestWriteUniqueId, NULL, bOutSuccess);
-				}
-
-				FMemoryReader TestReadUniqueId(Buffer);
-
-				FUniqueNetIdRepl EmptyIdOut;
-				TestReadUniqueId << EmptyIdOut;
-				if (EmptyIdOut.GetUniqueNetId().IsValid())
-				{
-					UE_LOG(LogNet, Warning, TEXT("EmptyId %s should have been invalid"), *EmptyIdOut->ToString());
-					bSuccess = false;
-				}
-
-				FUniqueNetIdRepl ValidIdOut;
-				TestReadUniqueId << ValidIdOut;
-				if (*UserId != *ValidIdOut.GetUniqueNetId())
-				{
-					UE_LOG(LogNet, Warning, TEXT("UserId input %s != UserId output %s"), *ValidIdIn->ToString(), *ValidIdOut->ToString());
-					bSuccess = false;
-				}
+				// Net serialize
+				bool bOutSuccess = false;
+				EmptyIdIn.NetSerialize(TestWriteUniqueId, NULL, bOutSuccess);
+				ValidIdIn.NetSerialize(TestWriteUniqueId, NULL, bOutSuccess);
+			}
+
+			FMemoryReader TestReadUniqueId(Buffer);
+
+			FUniqueNetIdRepl EmptyIdOut;
+			TestReadUniqueId << EmptyIdOut;
+			if (EmptyIdOut.GetUniqueNetId().IsValid())
+			{
+				UE_LOG(LogNet, Warning, TEXT("EmptyId %s should have been invalid"), *EmptyIdOut->ToString());
+				bSuccess = false;
+			}
+
+			FUniqueNetIdRepl ValidIdOut;
+			TestReadUniqueId << ValidIdOut;
+			if (*UserId != *ValidIdOut.GetUniqueNetId())
+			{
+				UE_LOG(LogNet, Warning, TEXT("UserId input %s != UserId output %s"), *ValidIdIn->ToString(), *ValidIdOut->ToString());
+				bSuccess = false;
 			}
 		}
+	}
 
+	if (bSuccess)
+	{
+		FString OutString;
+		TSharedRef<FJsonValue> JsonValue = ValidIdIn.ToJson();
+		bSuccess = JsonValue->TryGetString(OutString);
 		if (bSuccess)
 		{
-			FString OutString;
-			TSharedRef<FJsonValue> JsonValue = ValidIdIn.ToJson();
-			bSuccess = JsonValue->TryGetString(OutString);
-			if (bSuccess)
-			{
-				FUniqueNetIdRepl NewIdOut;
-				NewIdOut.FromJson(OutString);
-				bSuccess = NewIdOut.IsValid();
-			}
+			FUniqueNetIdRepl NewIdOut;
+			NewIdOut.FromJson(OutString);
+			bSuccess = NewIdOut.IsValid();
 		}
 	}
 
+
 	if (!bSuccess)
 	{
 		UE_LOG(LogNet, Warning, TEXT("TestUniqueIdRepl test failure!"));
diff --git a/Engine/Source/Runtime/Engine/Private/OnlineSession.cpp b/Engine/Source/Runtime/Engine/Private/OnlineSession.cpp
--- a/Engine/Source/Runtime/Engine/Private/OnlineSession.cpp
+++ b/Engine/Source/Runtime/Engine/Private/OnlineSession.cpp
@@ -1,7 +1,8 @@
 // Copyright 1998-2016 Epic Games, Inc. All Rights Reserved.
 
 /*=============================================================================
-	GameSession.cpp: GameSession code.
+	OnlineSession.cpp: Online session related implementations
+	(creating/joining/leaving/destroying sessions)
 =============================================================================*/
 
 #include "EnginePrivate.h"
diff --git a/Engine/Source/Runtime/Engine/Private/PlayerController.cpp b/Engine/Source/Runtime/Engine/Private/PlayerController.cpp
--- a/Engine/Source/Runtime/Engine/Private/PlayerController.cpp
+++ b/Engine/Source/Runtime/Engine/Private/PlayerController.cpp
@@ -9,7 +9,7 @@
 #include "Interfaces/NetworkPredictionInterface.h"
 #include "ConfigCacheIni.h"
 #include "SoundDefinitions.h"
-#include "OnlineSubsystemUtils.h"
+#include "Net/OnlineEngineInterface.h"
 #include "GameFramework/OnlineSession.h"
 #include "IHeadMountedDisplay.h"
 #include "IMotionController.h"
@@ -3007,17 +3007,13 @@ void APlayerController::ToggleSpeaking(bool bSpeaking)
 	if (LP != NULL)
 	{
 		UWorld* World = GetWorld();
-		IOnlineVoicePtr VoiceInt = Online::GetVoiceInterface(World);
-		if (VoiceInt.IsValid())
+		if (bSpeaking)
 		{
-			if (bSpeaking)
-			{
-				VoiceInt->StartNetworkedVoice(LP->GetControllerId());
-			}
-			else
-			{
-				VoiceInt->StopNetworkedVoice(LP->GetControllerId());
-			}
+			UOnlineEngineInterface::Get()->StartNetworkedVoice(World, LP->GetControllerId());
+		}
+		else
+		{
+			UOnlineEngineInterface::Get()->StopNetworkedVoice(World, LP->GetControllerId());
 		}
 	}
 }
