Quail: fix online serialization

diff --git a/Engine/Plugins/Online/OnlineSubsystem/Source/Public/OnlineSubsystemNames.h b/Engine/Plugins/Online/OnlineSubsystem/Source/Public/OnlineSubsystemNames.h
--- a/Engine/Plugins/Online/OnlineSubsystem/Source/Public/OnlineSubsystemNames.h
+++ b/Engine/Plugins/Online/OnlineSubsystem/Source/Public/OnlineSubsystemNames.h
@@ -92,4 +92,8 @@
 
 #ifndef OCULUS_SUBSYSTEM
 #define OCULUS_SUBSYSTEM FName(TEXT("Oculus"))
-#endif
\ No newline at end of file
+#endif
+
+#ifndef QUAIL_SUBSYSTEM
+#define QUAIL_SUBSYSTEM FName(TEXT("Quail"))
+#endif
diff --git a/Engine/Plugins/Online/OnlineSubsystemUtils/Source/OnlineSubsystemUtils/Private/OnlineSubsystemUtilsModule.cpp b/Engine/Plugins/Online/OnlineSubsystemUtils/Source/OnlineSubsystemUtils/Private/OnlineSubsystemUtilsModule.cpp
--- a/Engine/Plugins/Online/OnlineSubsystemUtils/Source/OnlineSubsystemUtils/Private/OnlineSubsystemUtilsModule.cpp
+++ b/Engine/Plugins/Online/OnlineSubsystemUtils/Source/OnlineSubsystemUtils/Private/OnlineSubsystemUtilsModule.cpp
@@ -209,6 +209,7 @@ private:
 		CREATE_HASH(WECHAT_SUBSYSTEM); 
 		CREATE_HASH(TWITCH_SUBSYSTEM); 
 		CREATE_HASH(OCULUS_SUBSYSTEM); 
+		CREATE_HASH(QUAIL_SUBSYSTEM);
 		// Shouldn't need these as they are mocking interfaces for existing platforms
 		CREATE_HASH(PS4SERVER_SUBSYSTEM);
 		CREATE_HASH(LIVESERVER_SUBSYSTEM);
