Add a timeout to the arbitrary rich presence fetch test. If this happens, the task is a failure.

diff --git a/Engine/Plugins/Online/OnlineSubsystemUtils/Source/OnlineSubsystemUtils/Private/Tests/TestPresenceInterface.cpp b/Engine/Plugins/Online/OnlineSubsystemUtils/Source/OnlineSubsystemUtils/Private/Tests/TestPresenceInterface.cpp
--- a/Engine/Plugins/Online/OnlineSubsystemUtils/Source/OnlineSubsystemUtils/Private/Tests/TestPresenceInterface.cpp
+++ b/Engine/Plugins/Online/OnlineSubsystemUtils/Source/OnlineSubsystemUtils/Private/Tests/TestPresenceInterface.cpp
@@ -7,6 +7,7 @@
 #if WITH_DEV_AUTOMATION_TESTS
 
 #define PRESENCE_TEST_LOG_KEY 50
+#define PRESENCE_TEST_TIMEOUT 120
 
 FTestPresenceInterface::~FTestPresenceInterface()
 {
@@ -95,7 +96,7 @@ bool FTestPresenceInterface::Tick(float DeltaTime)
 
 	// If we don't get an update within 2min while this test runs, mark the delegate watch task as a failure
 	if (!EnumHasAnyFlags(CompletedTasks, EPresenceTestStatus::WorkingDelegate) &&
-		FPlatformTime::Seconds() - TestTimeStart > 120 &&
+		FPlatformTime::Seconds() - TestTimeStart > PRESENCE_TEST_TIMEOUT &&
 		TestTimeStart > 0)
 	{
 		bHasFailed = true;
@@ -103,6 +104,17 @@ bool FTestPresenceInterface::Tick(float DeltaTime)
 		TestTimeStart = 0; // Reset the time so we don't run the task again
 	}
 
+	// Handle failing out the arbitrary fetch test
+	if (!EnumHasAnyFlags(CompletedTasks, EPresenceTestStatus::FetchRandom) &&
+		EnumHasAnyFlags(RequiredFlags, EPresenceTestStatus::FetchRandom) &&
+		FPlatformTime::Seconds() - ArbitraryFetchTimeStart > PRESENCE_TEST_TIMEOUT &&
+		ArbitraryFetchTimeStart > 0)
+	{
+		bHasFailed = true;
+		TasksAttempted |= EPresenceTestStatus::FetchRandom;
+		ArbitraryFetchTimeStart = 0;
+	}
+
 	if (EnumHasAllFlags(CompletedTasks, RequiredFlags) || (bHasFailed && EnumHasAllFlags(TasksAttempted, RequiredFlags)))
 	{
 		UE_LOG_ONLINE_PRESENCE(Log, TEXT("Presence testing completed! Success: %d Failed Tasks: %s"), !bHasFailed, *PrintTestFailure());
@@ -196,7 +208,7 @@ void FTestPresenceInterface::Test(UWorld* InWorld, const FString& RandomUser)
 
 	World = InWorld;
 	// This test typically times out after two minutes. This timer is set for a bit longer than that, just in case.
-	GEngine->AddOnScreenDebugMessage(PRESENCE_TEST_LOG_KEY, 125.0f, FColor::Yellow, TEXT("Presence Test is running. Please wait until it completes."));
+	GEngine->AddOnScreenDebugMessage(PRESENCE_TEST_LOG_KEY, 145.0f, FColor::Yellow, TEXT("Presence Test is running. Please wait until it completes."));
 	GEngine->OnWorldDestroyed().AddRaw(this, &FTestPresenceInterface::OnWorldDestroyed);
 
 	// Grab the user's friends
@@ -213,6 +225,7 @@ void FTestPresenceInterface::Test(UWorld* InWorld, const FString& RandomUser)
 		if (ArbitraryId.IsValid())
 		{
 			PresenceInt->QueryPresence(*ArbitraryId, IOnlinePresence::FOnPresenceTaskCompleteDelegate::CreateRaw(this, &FTestPresenceInterface::OnRandomUserFetchComplete));
+			ArbitraryFetchTimeStart = FPlatformTime::Seconds();
 			RequiredFlags |= EPresenceTestStatus::FetchRandom;
 		}
 		else
diff --git a/Engine/Plugins/Online/OnlineSubsystemUtils/Source/OnlineSubsystemUtils/Private/Tests/TestPresenceInterface.h b/Engine/Plugins/Online/OnlineSubsystemUtils/Source/OnlineSubsystemUtils/Private/Tests/TestPresenceInterface.h
--- a/Engine/Plugins/Online/OnlineSubsystemUtils/Source/OnlineSubsystemUtils/Private/Tests/TestPresenceInterface.h
+++ b/Engine/Plugins/Online/OnlineSubsystemUtils/Source/OnlineSubsystemUtils/Private/Tests/TestPresenceInterface.h
@@ -41,7 +41,8 @@ public:
 		TasksAttempted(EPresenceTestStatus::None),
 		bHasCreatedSession(false),
 		bHasFailed(false),
-		TestTimeStart(0.0)
+		TestTimeStart(0.0),
+		ArbitraryFetchTimeStart(0.0)
 	{
 	}
 
@@ -88,6 +89,7 @@ private:
 	bool bHasCreatedSession;
 	bool bHasFailed;
 	double TestTimeStart;
+	double ArbitraryFetchTimeStart;
 
 	/** Cached presence object for testing if the presences are the same */
 	FOnlineUserPresenceStatus CachedPresence;
