Fix IpConnection resolve nuking previous data set, which would wipe away port information. Now we cache off that port data before the clone.

diff --git a/Engine/Plugins/Online/OnlineSubsystemUtils/Source/OnlineSubsystemUtils/Private/IpConnection.cpp b/Engine/Plugins/Online/OnlineSubsystemUtils/Source/OnlineSubsystemUtils/Private/IpConnection.cpp
--- a/Engine/Plugins/Online/OnlineSubsystemUtils/Source/OnlineSubsystemUtils/Private/IpConnection.cpp
+++ b/Engine/Plugins/Online/OnlineSubsystemUtils/Source/OnlineSubsystemUtils/Private/IpConnection.cpp
@@ -128,7 +128,10 @@ void UIpConnection::LowLevelSend(void* Data, int32 CountBytes, int32 CountBits)
 		else
 		{
 			// Host name resolution just now succeeded.
+			int32 CurPort = RemoteAddr->GetPort();
 			RemoteAddr = ResolveInfo->GetResolvedAddress().Clone();
+			RemoteAddr->SetPort(CurPort);
+
 			UE_LOG(LogNet, Log, TEXT("Host name resolution completed"));
 			delete ResolveInfo;
 			ResolveInfo = NULL;
