Fix for failed check in FNetPacketNotify after a socket failure.
Fix for crash in SendBunch after socket failure

Postponed call to NetConnection::Close() to next tick if a socket error occurs durintg send to avoid issues with reentrant code paths in DataChannel::SendBunch() and NetConnection::FlushNet()

diff --git a/Engine/Plugins/Online/OnlineSubsystemUtils/Source/OnlineSubsystemUtils/Private/IpConnection.cpp b/Engine/Plugins/Online/OnlineSubsystemUtils/Source/OnlineSubsystemUtils/Private/IpConnection.cpp
--- a/Engine/Plugins/Online/OnlineSubsystemUtils/Source/OnlineSubsystemUtils/Private/IpConnection.cpp
+++ b/Engine/Plugins/Online/OnlineSubsystemUtils/Source/OnlineSubsystemUtils/Private/IpConnection.cpp
@@ -290,18 +290,16 @@ void UIpConnection::HandleSocketSendResult(const FSocketSendResult& Result, ISoc
 	if (Result.Error != SE_EWOULDBLOCK &&
 		Result.Error != SE_NO_ERROR)
 	{
-		FString ErrorString = FString::Printf(TEXT("UIpNetConnection::LowLevelSend: Socket->SendTo failed with error %i (%s). %s"),
+		FString ErrorString = FString::Printf(TEXT("UIpNetConnection::LowLevelSend: Socket->SendTo failed with error %i (%s). %s Connection will be closed during next Tick()!"),
 			static_cast<int32>(Result.Error),
 			SocketSubsystem->GetSocketError(Result.Error),
 			*Describe());
 
 		GEngine->BroadcastNetworkFailure(Driver->GetWorld(), Driver, ENetworkFailure::ConnectionLost, ErrorString);
 
-		// Reset the send buffer before closing, as it could have been (almost) full and the close process may
-		// write a bunch that could cause an overflow.  We're closing the connection anyway, and given that
-		// the socket is returning errors, the close bunch probably won't be delivered either.
-		InitSendBuffer();
-		Close();
+		// Request the connection to be disconnected during next tick() since we got a critical socket failure, the actual disconnect is postponed 		
+		// to avoid issues with the call Close() causing issues with reentrant code paths in DataChannel::SendBunch() and FlushNet()
+		SetPendingCloseDueToSocketSendFailure();
 	}
 }
 
diff --git a/Engine/Source/Runtime/Engine/Classes/Engine/NetConnection.h b/Engine/Source/Runtime/Engine/Classes/Engine/NetConnection.h
--- a/Engine/Source/Runtime/Engine/Classes/Engine/NetConnection.h
+++ b/Engine/Source/Runtime/Engine/Classes/Engine/NetConnection.h
@@ -1064,6 +1064,8 @@ public:
 
 protected:
 
+	void SetPendingCloseDueToSocketSendFailure();
+
 	void CleanupDormantActorState();
 
 	/** Called internally to destroy an actor during replay fast-forward when the actor channel index will be recycled */
@@ -1140,6 +1142,9 @@ private:
 	/** True if we've hit the actor channel limit and logged a warning about it */
 	bool bHasWarnedAboutChannelLimit;
 
+	/** True if we are pending close due to a socket failure during send */
+	bool bConnectionPendingCloseDueToSocketSendFailure;
+
 	FNetworkGUID GetActorGUIDFromOpenBunch(FInBunch& Bunch);
 };
 
diff --git a/Engine/Source/Runtime/Engine/Private/NetConnection.cpp b/Engine/Source/Runtime/Engine/Private/NetConnection.cpp
--- a/Engine/Source/Runtime/Engine/Private/NetConnection.cpp
+++ b/Engine/Source/Runtime/Engine/Private/NetConnection.cpp
@@ -163,6 +163,7 @@ PRAGMA_ENABLE_DEPRECATION_WARNINGS
 ,	LastNotifiedPacketId( -1 )
 ,	HasDirtyAcks(0u)
 ,	bHasWarnedAboutChannelLimit(false)
+,	bConnectionPendingCloseDueToSocketSendFailure(false)
 {
 	// This isn't ideal, because it won't capture memory derived classes are creating dynamically.
 	// The allocations could *probably* be moved somewhere else (like InitBase), but that
@@ -2619,6 +2620,15 @@ void UNetConnection::Tick()
 		OutPackets = 0;
 	}
 
+	if (bConnectionPendingCloseDueToSocketSendFailure)
+	{
+		Close();
+		bConnectionPendingCloseDueToSocketSendFailure = false;
+
+		// early out
+		return;
+	}
+
 	// Compute time passed since last update.
 	const float DeltaTime	= Driver->Time - LastTickTime;
 	LastTickTime			= Driver->Time;
@@ -3298,6 +3308,12 @@ void UNetConnection::CleanupStaleDormantReplicators()
 	}
 }
 
+
+void UNetConnection::SetPendingCloseDueToSocketSendFailure()
+{
+	bConnectionPendingCloseDueToSocketSendFailure = true;
+}
+
 /*-----------------------------------------------------------------------------
 	USimulatedClientNetConnection.
 -----------------------------------------------------------------------------*/
