Refactored TickDispatch to use FPacketIterator for receiving packets - leftover refactor from out of order packet fix.

diff --git a/Engine/Plugins/Online/OnlineSubsystemUtils/Source/OnlineSubsystemUtils/Classes/IpNetDriver.h b/Engine/Plugins/Online/OnlineSubsystemUtils/Source/OnlineSubsystemUtils/Classes/IpNetDriver.h
--- a/Engine/Plugins/Online/OnlineSubsystemUtils/Source/OnlineSubsystemUtils/Classes/IpNetDriver.h
+++ b/Engine/Plugins/Online/OnlineSubsystemUtils/Source/OnlineSubsystemUtils/Classes/IpNetDriver.h
@@ -23,6 +23,8 @@ class FSocket;
 UCLASS(transient, config=Engine)
 class ONLINESUBSYSTEMUTILS_API UIpNetDriver : public UNetDriver
 {
+	friend class FPacketIterator;
+
     GENERATED_UCLASS_BODY()
 
 	/** Should port unreachable messages be logged */
@@ -109,6 +111,17 @@ class ONLINESUBSYSTEMUTILS_API UIpNetDriver : public UNetDriver
 	/** @return TCPIP connection to server */
 	class UIpConnection* GetServerConnection();
 
+private:
+	/**
+	 * Whether or not a socket receive failure Error indicates a blocking error, which should 'break;' the receive loop
+	 */
+	static FORCEINLINE bool IsRecvFailBlocking(ESocketErrors Error)
+	{
+		// SE_ECONNABORTED is for PS4 LAN cable pulls
+		return Error == SE_NO_ERROR || Error == SE_EWOULDBLOCK || Error == SE_ECONNABORTED;
+	};
+
+public:
 	// Callback for platform handling when networking is taking a long time in a single frame (by default over 1 second).
 	// It may get called multiple times in a single frame if additional processing after a previous alert exceeds the threshold again
 	DECLARE_MULTICAST_DELEGATE(FOnNetworkProcessingCausingSlowFrame);
diff --git a/Engine/Plugins/Online/OnlineSubsystemUtils/Source/OnlineSubsystemUtils/Private/IpNetDriver.cpp b/Engine/Plugins/Online/OnlineSubsystemUtils/Source/OnlineSubsystemUtils/Private/IpNetDriver.cpp
--- a/Engine/Plugins/Online/OnlineSubsystemUtils/Source/OnlineSubsystemUtils/Private/IpNetDriver.cpp
+++ b/Engine/Plugins/Online/OnlineSubsystemUtils/Source/OnlineSubsystemUtils/Private/IpNetDriver.cpp
@@ -74,6 +74,251 @@ TAutoConsoleVariable<int32> CVarNetIpNetDriverReceiveThreadPollTimeMS(
 	250,
 	TEXT("If net.IpNetDriverUseReceiveThread is true, the number of milliseconds to use as the timeout value for FSocket::Wait on the receive thread. A negative value means to wait indefinitely (FSocket::Shutdown should cancel it though)."));
 
+
+
+/**
+ * FPacketItrator
+ *
+ * Encapsulates the NetDriver TickDispatch code required for executing all variations of packet receives
+ * (FSocket::RecvFrom, FSocket::RecvMulti, and the Receive Thread),
+ * as well as implementing/abstracting-away some of the outermost (non-NetConnection-related) parts of the DDoS detection code.
+ */
+class FPacketIterator
+{
+	friend class UIpNetDriver;
+	
+private:
+	struct FCachedPacket
+	{
+		/** Whether or not socket receive succeeded. Don't rely on the Error field for this, due to implementation/platform uncertainties. */
+		bool bRecvSuccess;
+
+		/** Pre-allocated Data field, for storing packets of any expected size */
+		TArray<uint8, TFixedAllocator<MAX_PACKET_SIZE>> Data;
+
+		/** Receive address for the packet */
+		TSharedPtr<FInternetAddr> Address;
+
+		/** OS-level timestamp for the packet receive, if applicable */
+		double PacketTimestamp;
+
+		/** Error if receiving a packet failed */
+		ESocketErrors Error;
+	};
+	
+private:
+	FPacketIterator(UIpNetDriver* InDriver)
+		: bBreak(false)
+		, Driver(InDriver)
+		, DDoS(Driver->DDoS)
+		, SocketSubsystem(Driver->GetSocketSubsystem())
+		, Socket(Driver->Socket)
+		, SocketReceiveThreadRunnable(Driver->SocketReceiveThreadRunnable.Get())
+		, CurrentPacket()
+	{
+		CurrentPacket.Address = SocketSubsystem->CreateInternetAddr();
+
+		AdvanceCurrentPacket();
+	}
+
+	FORCEINLINE FPacketIterator& operator++()
+	{
+		AdvanceCurrentPacket();
+
+		return *this;
+	}
+
+	FORCEINLINE explicit operator bool() const
+	{
+		return !bBreak;
+	}
+
+
+	/**
+	 * Retrieves the packet information from the current iteration. Avoid calling more than once, per iteration.
+	 *
+	 * @param OutData		Returns a pointer to the packet data
+	 * @param OutBytesRead	Returns the number of bytes read
+	 * @param OutError		Returns the socket error, if receiving failed
+	 * @param OutAddr		Returns the address the packet was received from
+	 * @return				Returns whether or not receiving was successful for the current packet
+	 */
+	bool GetCurrentPacket(const uint8*& OutData, uint32& OutBytesRead, ESocketErrors& OutError, TSharedPtr<FInternetAddr>& OutAddr)
+	{
+		bool bRecvSuccess = false;
+
+		OutData = CurrentPacket.Data.GetData();
+		OutBytesRead = CurrentPacket.Data.Num();
+		OutError = CurrentPacket.Error;
+		OutAddr = CurrentPacket.Address;
+		bRecvSuccess = CurrentPacket.bRecvSuccess;
+
+		return bRecvSuccess;
+	}
+
+	/**
+	 * Advances the current packet to the next iteration
+	 */
+	void AdvanceCurrentPacket()
+	{
+		bBreak = !ReceiveSinglePacket();
+	}
+
+	/**
+	 * Receives a single packet from the network socket, outputting to the CurrentPacket buffer.
+	 *
+	 * @return				Whether or not a packet or an error was successfully received
+	 */
+	bool ReceiveSinglePacket()
+	{
+		bool bReceivedPacketOrError = false;
+
+		CurrentPacket.bRecvSuccess = false;
+		CurrentPacket.Data.SetNumUninitialized(0, false);
+		CurrentPacket.Address->SetAnyAddress();
+		CurrentPacket.PacketTimestamp = 0.0;
+		CurrentPacket.Error = SE_NO_ERROR;
+
+		while (true)
+		{
+			bReceivedPacketOrError = false;
+
+			if (SocketReceiveThreadRunnable != nullptr)
+			{
+				// Very-early-out - the NetConnection per frame time limit, limits all packet processing
+				// @todo #JohnB: This DDoS detection code will be redundant, as it's performed in the Receive Thread in a coming refactor
+				if (DDoS.ShouldBlockNetConnPackets())
+				{
+					// Approximate due to threading
+					uint32 DropCountApprox = SocketReceiveThreadRunnable->ReceiveQueue.Count();
+
+					SocketReceiveThreadRunnable->ReceiveQueue.Empty();
+
+					if (DropCountApprox > 0)
+					{
+						DDoS.IncDroppedPacketCounter(DropCountApprox);
+					}
+				}
+				else
+				{
+					UIpNetDriver::FReceivedPacket IncomingPacket;
+					const bool bHasPacket = SocketReceiveThreadRunnable->ReceiveQueue.Dequeue(IncomingPacket);
+
+					if (bHasPacket)
+					{
+						if (IncomingPacket.FromAddress.IsValid())
+						{
+							CurrentPacket.Address = IncomingPacket.FromAddress.ToSharedRef();
+						}
+
+						ESocketErrors CurError = IncomingPacket.Error;
+						bool bReceivedPacket = CurError == SE_NO_ERROR;
+
+						CurrentPacket.bRecvSuccess = bReceivedPacket;
+						CurrentPacket.PacketTimestamp = IncomingPacket.PlatformTimeSeconds;
+						CurrentPacket.Error = CurError;
+						bReceivedPacketOrError = bReceivedPacket;
+
+						if (bReceivedPacket)
+						{
+							int32 BytesRead = IncomingPacket.PacketBytes.Num();
+
+							if (IncomingPacket.PacketBytes.Num() <= MAX_PACKET_SIZE)
+							{
+								CurrentPacket.Data.SetNumUninitialized(BytesRead, false);
+
+								FMemory::Memcpy(CurrentPacket.Data.GetData(), IncomingPacket.PacketBytes.GetData(), BytesRead);
+							}
+							else
+							{
+								UE_LOG(LogNet, Warning, TEXT("IpNetDriver receive thread received a packet of %d bytes, which is larger than the data buffer size of %d bytes."),
+										BytesRead, MAX_PACKET_SIZE);
+
+								continue;
+							}
+						}
+						// Received an error
+						else if (!UIpNetDriver::IsRecvFailBlocking(CurError))
+						{
+							bReceivedPacketOrError = true;
+						}
+					}
+				}
+			}
+			else
+			{
+				SCOPE_CYCLE_COUNTER(STAT_IpNetDriver_RecvFromSocket);
+
+				int32 BytesRead = 0;
+				bool bReceivedPacket = Socket->RecvFrom(CurrentPacket.Data.GetData(), MAX_PACKET_SIZE, BytesRead, *CurrentPacket.Address);
+
+				CurrentPacket.bRecvSuccess = bReceivedPacket;
+				bReceivedPacketOrError = bReceivedPacket;
+
+				if (bReceivedPacket)
+				{
+					// Fixed allocator, so no risk of realloc from copy-then-resize
+					CurrentPacket.Data.SetNumUninitialized(BytesRead, false);
+				}
+				else
+				{
+					ESocketErrors CurError = SocketSubsystem->GetLastErrorCode();
+
+					CurrentPacket.Error = CurError;
+					CurrentPacket.Data.SetNumUninitialized(0, false);
+
+					// Received an error
+					if (!UIpNetDriver::IsRecvFailBlocking(CurError))
+					{
+						bReceivedPacketOrError = true;
+					}
+				}
+
+				// Very-early-out - the NetConnection per frame time limit, limits all packet processing
+				if (bReceivedPacketOrError && DDoS.ShouldBlockNetConnPackets())
+				{
+					if (bReceivedPacket)
+					{
+						DDoS.IncDroppedPacketCounter();
+					}
+
+					continue;
+				}
+			}
+
+			// While loop only exists to allow 'continue' for DDoS and invalid packet code, above
+			break;
+		}
+
+		return bReceivedPacketOrError;
+	}
+
+
+private:
+	/** Specified internally, when the packet iterator should break/stop (no packets, DDoS limits triggered, etc.) */
+	bool bBreak;
+
+
+	/** Cached reference to the NetDriver, and NetDriver variables/values */
+
+	UIpNetDriver* Driver;
+
+	FDDoSDetection& DDoS;
+
+	ISocketSubsystem* SocketSubsystem;
+
+	FSocket*& Socket;
+
+	UIpNetDriver::FReceiveThreadRunnable* SocketReceiveThreadRunnable;
+
+	/** Stores information for the current packet being received (when using single-receive mode) */
+	FCachedPacket CurrentPacket;
+};
+
+/**
+ * UIpNetDriver
+ */
+
 UIpNetDriver::UIpNetDriver(const FObjectInitializer& ObjectInitializer)
 	: Super(ObjectInitializer)
 	, PauseReceiveEnd(0.f)
@@ -277,19 +522,17 @@ void UIpNetDriver::TickDispatch(float DeltaTime)
 
 	FScopedLevelCollectionContextSwitch LCSwitch(FoundCollectionIndex, World);
 
-	ISocketSubsystem* SocketSubsystem = GetSocketSubsystem();
 
 	DDoS.PreFrameReceive(DeltaTime);
 
+	ISocketSubsystem* SocketSubsystem = GetSocketSubsystem();
 	const double StartReceiveTime = FPlatformTime::Seconds();
 	double AlarmTime = StartReceiveTime + GIpNetDriverMaxDesiredTimeSliceBeforeAlarmSecs;
 	const bool bSlowFrameChecks = OnNetworkProcessingCausingSlowFrame.IsBound();
 
-	// Process all incoming packets.
-	uint8 Data[MAX_PACKET_SIZE];
-	TSharedRef<FInternetAddr> FromAddr = SocketSubsystem->CreateInternetAddr();
 
-	for (; Socket != nullptr;)
+	// Process all incoming packets
+	for (FPacketIterator It(this); It; ++It)
 	{
 		// @todo: Remove the slow frame checks, eventually - potential DDoS and Switch platform constraint
 		if (bSlowFrameChecks)
@@ -303,49 +546,12 @@ void UIpNetDriver::TickDispatch(float DeltaTime)
 			}
 		}
 
+		uint8* Data = nullptr;
 		int32 BytesRead = 0;
-
-		// Reset the address on every pass. Otherwise if there's an error receiving, the address may be from a previous packet.
-		FromAddr->SetAnyAddress();
-
-		// Get data, if any.
-		bool bOk = false;
 		ESocketErrors Error = SE_NO_ERROR;
-		bool bUsingReceiveThread = SocketReceiveThreadRunnable.IsValid();
-
-		if (bUsingReceiveThread)
-		{
-			FReceivedPacket IncomingPacket;
-			const bool bHasPacket = SocketReceiveThreadRunnable->ReceiveQueue.Dequeue(IncomingPacket);
-			if (!bHasPacket)
-			{
-				break;
-			}
-
-			if (IncomingPacket.FromAddress.IsValid())
-			{
-				FromAddr = IncomingPacket.FromAddress.ToSharedRef();
-			}
-			Error = IncomingPacket.Error;
-			bOk = Error == SE_NO_ERROR;
-			
-			if (IncomingPacket.PacketBytes.Num() <= sizeof(Data))
-			{
-				FMemory::Memcpy(Data, IncomingPacket.PacketBytes.GetData(), IncomingPacket.PacketBytes.Num());
-				BytesRead = IncomingPacket.PacketBytes.Num();
-			}
-			else
-			{
-				UE_LOG(LogNet, Log, TEXT("IpNetDriver receive thread received a packet of %d bytes, which is larger than the data buffer size of %d bytes."), IncomingPacket.PacketBytes.Num(), sizeof(Data));
-				continue;
-			}
-		}
-		else
-		{
-			SCOPE_CYCLE_COUNTER(STAT_IpNetDriver_RecvFromSocket);
-			bOk = Socket->RecvFrom(Data, sizeof(Data), BytesRead, *FromAddr);
-		}
-
+		TSharedPtr<FInternetAddr> FromAddrPtr;
+		bool bOk = It.GetCurrentPacket((const uint8*&)Data, (uint32&)BytesRead, Error, FromAddrPtr);
+		TSharedRef<FInternetAddr> FromAddr = FromAddrPtr.ToSharedRef();
 		UNetConnection* Connection = nullptr;
 		UIpConnection* const MyServerConnection = GetServerConnection();
 
@@ -362,12 +568,7 @@ void UIpNetDriver::TickDispatch(float DeltaTime)
 		}
 		else
 		{
-			if (!bUsingReceiveThread)
-			{
-				Error = SocketSubsystem->GetLastErrorCode();
-			}
-
-			if (Error == SE_EWOULDBLOCK || Error == SE_NO_ERROR || Error == SE_ECONNABORTED)
+			if (IsRecvFailBlocking(Error))
 			{
 				// No data or no error? (SE_ECONNABORTED is for PS4 LAN cable pulls)
 				break;
@@ -429,16 +630,6 @@ void UIpNetDriver::TickDispatch(float DeltaTime)
 			}
 		}
 
-		// Very-early-out - the NetConnection per frame time limit, limits all packet processing
-		if (DDoS.ShouldBlockNetConnPackets())
-		{
-			if (bOk)
-			{
-				DDoS.IncDroppedPacketCounter();
-			}
-
-			continue;
-		}
 
 		// Figure out which socket the received data came from.
 		if (MyServerConnection)
@@ -837,7 +1028,8 @@ void UIpNetDriver::LowLevelDestroy()
 		}
 		// Free the memory the OS allocated for this socket
 		SocketSubsystem->DestroySocket(Socket);
-		Socket = NULL;
+		Socket = nullptr;
+
 		UE_LOG(LogExit, Log, TEXT("%s shut down"),*GetDescription() );
 	}
 
@@ -975,7 +1167,7 @@ uint32 UIpNetDriver::FReceiveThreadRunnable::Run()
 				IncomingPacket.Error = SocketSubsystem->GetLastErrorCode();
 
 				// Pass all other errors back to the Game Thread
-				if (IncomingPacket.Error == SE_EWOULDBLOCK || IncomingPacket.Error == SE_NO_ERROR || IncomingPacket.Error == SE_ECONNABORTED)
+				if (IsRecvFailBlocking(IncomingPacket.Error))
 				{
 					continue;
 				}
