Added additional checks against nullptr dereference.

diff --git a/Engine/Plugins/Online/OnlineSubsystemUtils/Source/OnlineSubsystemUtils/Private/IpNetDriver.cpp b/Engine/Plugins/Online/OnlineSubsystemUtils/Source/OnlineSubsystemUtils/Private/IpNetDriver.cpp
--- a/Engine/Plugins/Online/OnlineSubsystemUtils/Source/OnlineSubsystemUtils/Private/IpNetDriver.cpp
+++ b/Engine/Plugins/Online/OnlineSubsystemUtils/Source/OnlineSubsystemUtils/Private/IpNetDriver.cpp
@@ -116,7 +116,10 @@ private:
 		, SocketReceiveThreadRunnable(Driver->SocketReceiveThreadRunnable.Get())
 		, CurrentPacket()
 	{
-		CurrentPacket.Address = SocketSubsystem->CreateInternetAddr();
+		if (SocketSubsystem != nullptr)
+		{
+			CurrentPacket.Address = SocketSubsystem->CreateInternetAddr();
+		}
 
 		AdvanceCurrentPacket();
 	}
@@ -175,7 +178,12 @@ private:
 
 		CurrentPacket.bRecvSuccess = false;
 		CurrentPacket.Data.SetNumUninitialized(0, false);
-		CurrentPacket.Address->SetAnyAddress();
+
+		if (CurrentPacket.Address.IsValid())
+		{
+			CurrentPacket.Address->SetAnyAddress();
+		}
+
 		CurrentPacket.PacketTimestamp = 0.0;
 		CurrentPacket.Error = SE_NO_ERROR;
 
@@ -245,7 +253,7 @@ private:
 					}
 				}
 			}
-			else
+			else if (Socket != nullptr && SocketSubsystem != nullptr)
 			{
 				SCOPE_CYCLE_COUNTER(STAT_IpNetDriver_RecvFromSocket);
 
