Allow socket errors to happen for a configured delay before killing connections.
In some cases, socket errors may be recoverable and we're being overly aggressive.

diff --git a/Engine/Plugins/Online/OnlineSubsystemUtils/Source/OnlineSubsystemUtils/Classes/IpConnection.h b/Engine/Plugins/Online/OnlineSubsystemUtils/Source/OnlineSubsystemUtils/Classes/IpConnection.h
--- a/Engine/Plugins/Online/OnlineSubsystemUtils/Source/OnlineSubsystemUtils/Classes/IpConnection.h
+++ b/Engine/Plugins/Online/OnlineSubsystemUtils/Source/OnlineSubsystemUtils/Classes/IpConnection.h
@@ -42,6 +42,7 @@ class ONLINESUBSYSTEMUTILS_API UIpConnection : public UNetConnection
 	virtual FString RemoteAddressToString() override;
 	virtual void Tick() override;
 	virtual void CleanUp() override;
+	virtual void ReceivedRawPacket(void* Data, int32 Count) override;
 	//~ End NetConnection Interface
 
 	/**
@@ -79,6 +80,21 @@ private:
 	 */
 	FGraphEventRef LastSendTask;
 
+	/** Instead of disconnecting immediately on a socket error, wait for some time to see if we can recover. Specified in Seconds. */
+	UPROPERTY(Config)
+	float SocketErrorDisconnectDelay;
+
+	/** Cached time of the first send socket error that will be used to compute disconnect delay. */
+	float SocketError_SendDelayStartTime;
+
+	/** Cached time of the first recv socket error that will be used to compute disconnect delay. */
+	float SocketError_RecvDelayStartTime;
+
+	friend class FIpConnectionHelper;
+
 	/** Handles any SendTo errors on the game thread. */
 	void HandleSocketSendResult(const FSocketSendResult& Result, ISocketSubsystem* SocketSubsystem);
+
+	/** Notifies us that we've encountered an error while receiving a packet. */
+	void HandleSocketRecvError(class UNetDriver* NetDriver, const FString& ErrorString);
 };
diff --git a/Engine/Plugins/Online/OnlineSubsystemUtils/Source/OnlineSubsystemUtils/Private/IpConnection.cpp b/Engine/Plugins/Online/OnlineSubsystemUtils/Source/OnlineSubsystemUtils/Private/IpConnection.cpp
--- a/Engine/Plugins/Online/OnlineSubsystemUtils/Source/OnlineSubsystemUtils/Private/IpConnection.cpp
+++ b/Engine/Plugins/Online/OnlineSubsystemUtils/Source/OnlineSubsystemUtils/Private/IpConnection.cpp
@@ -39,7 +39,10 @@ UIpConnection::UIpConnection(const FObjectInitializer& ObjectInitializer) :
 	Super(ObjectInitializer),
 	RemoteAddr(NULL),
 	Socket(NULL),
-	ResolveInfo(NULL)
+	ResolveInfo(NULL),
+	SocketErrorDisconnectDelay(5.f),
+	SocketError_SendDelayStartTime(0.f),
+	SocketError_RecvDelayStartTime(0.f)
 {
 }
 
@@ -221,6 +224,8 @@ void UIpConnection::LowLevelSend(void* Data, int32 CountBits, FOutPacketTraits&
 
 		if (CountBytes > 0)
 		{
+			const bool bNotifyOnSuccess = (SocketErrorDisconnectDelay > 0.f) && (SocketError_SendDelayStartTime != 0.f);
+
 			if (CVarNetIpConnectionUseSendTasks.GetValueOnAnyThread() != 0)
 			{
 				DECLARE_CYCLE_STAT(TEXT("IpConnection SendTo task"), STAT_IpConnection_SendToTask, STATGROUP_TaskGraphTasks);
@@ -232,7 +237,8 @@ void UIpConnection::LowLevelSend(void* Data, int32 CountBits, FOutPacketTraits&
 				}
 
 				ISocketSubsystem* const SocketSubsystem = Driver->GetSocketSubsystem();
-				LastSendTask = FFunctionGraphTask::CreateAndDispatchWhenReady([this, Packet = TArray<uint8>(DataToSend, CountBytes), SocketSubsystem]
+				
+				LastSendTask = FFunctionGraphTask::CreateAndDispatchWhenReady([this, Packet = TArray<uint8>(DataToSend, CountBytes), SocketSubsystem, bNotifyOnSuccess]
 				{
 					bool bWasSendSuccessful = false;
 					UIpConnection::FSocketSendResult Result;
@@ -245,12 +251,12 @@ void UIpConnection::LowLevelSend(void* Data, int32 CountBits, FOutPacketTraits&
 					if (!bWasSendSuccessful && SocketSubsystem)
 					{
 						Result.Error = SocketSubsystem->GetLastErrorCode();
-						if (Result.Error != SE_EWOULDBLOCK &&
-							Result.Error != SE_NO_ERROR)
-						{
-							FScopeLock ScopeLock(&SocketSendResultsCriticalSection);
-							SocketSendResults.Add(MoveTemp(Result));
-						}
+					}
+
+					if (!bWasSendSuccessful || (bNotifyOnSuccess && Result.Error == SE_NO_ERROR))
+					{
+						FScopeLock ScopeLock(&SocketSendResultsCriticalSection);
+						SocketSendResults.Add(MoveTemp(Result));
 					}
 				},
 				GET_STATID(STAT_IpConnection_SendToTask), &Prerequisites);
@@ -267,11 +273,17 @@ void UIpConnection::LowLevelSend(void* Data, int32 CountBits, FOutPacketTraits&
 					SCOPE_CYCLE_COUNTER(STAT_IpConnection_SendToSocket);
 					bWasSendSuccessful = Socket->SendTo(DataToSend, CountBytes, SendResult.BytesSent, *RemoteAddr);
 				}
+
 				if (bWasSendSuccessful)
 				{
 					UNCLOCK_CYCLES(Driver->SendCycles);
 					NETWORK_PROFILER(GNetworkProfiler.FlushOutgoingBunches(this));
 					NETWORK_PROFILER(GNetworkProfiler.TrackSocketSendTo(Socket->GetDescription(), DataToSend, SendResult.BytesSent, NumPacketIdBits, NumBunchBits, NumAckBits, NumPaddingBits, this));
+
+					if (bNotifyOnSuccess)
+					{
+						HandleSocketSendResult(SendResult, nullptr);
+					}
 				}
 				else
 				{
@@ -285,12 +297,51 @@ void UIpConnection::LowLevelSend(void* Data, int32 CountBits, FOutPacketTraits&
 	}
 }
 
+void UIpConnection::ReceivedRawPacket(void* Data, int32 Count)
+{
+	UE_CLOG(SocketError_RecvDelayStartTime > 0.f, LogNet, Log, TEXT("UIpConnection::ReceivedRawPacket: Recoverd from socket errors. %s Connection"), *Describe());
+	
+	// We received data successfully, reset our error counters.
+	SocketError_RecvDelayStartTime = 0.f;
+
+	Super::ReceivedRawPacket(Data, Count);
+}
+
 void UIpConnection::HandleSocketSendResult(const FSocketSendResult& Result, ISocketSubsystem* SocketSubsystem)
 {
-	if (Result.Error != SE_EWOULDBLOCK &&
-		Result.Error != SE_NO_ERROR)
+	if (Result.Error == SE_NO_ERROR)
+	{
+		UE_CLOG(SocketError_SendDelayStartTime > 0.f, LogNet, Log, TEXT("UIpConnection::HandleSocketSendResult: Recovered from socket errors. %s Connection"), *Describe());
+
+		// We sent data successfully, reset our error counters.
+		SocketError_SendDelayStartTime = 0.f;
+	}
+	else if (Result.Error != SE_EWOULDBLOCK)
 	{
-		FString ErrorString = FString::Printf(TEXT("UIpNetConnection::LowLevelSend: Socket->SendTo failed with error %i (%s). %s Connection will be closed during next Tick()!"),
+		check(SocketSubsystem);
+
+		if (SocketErrorDisconnectDelay > 0.f)
+		{
+			const float Time = Driver->Time;
+			if (SocketError_SendDelayStartTime == 0.f)
+			{
+				UE_LOG(LogNet, Log, TEXT("UIpConnection::HandleSocketSendResult: Socket->SendTo failed with error %i (%s). %s Connection beginning close timeout (Timeout = %d)."),
+					static_cast<int32>(Result.Error),
+					SocketSubsystem->GetSocketError(Result.Error),
+					*Describe(),
+					SocketErrorDisconnectDelay);
+
+				SocketError_SendDelayStartTime = Time;
+				return;
+			}
+			else if ((Time - SocketError_SendDelayStartTime) < SocketErrorDisconnectDelay)
+			{
+				// Our delay hasn't elapsed yet. Just suppress further warnings until we either recover or are disconnected.
+				return;
+			}
+		}
+
+		FString ErrorString = FString::Printf(TEXT("UIpNetConnection::HandleSocketSendResult: Socket->SendTo failed with error %i (%s). %s Connection will be closed during next Tick()!"),
 			static_cast<int32>(Result.Error),
 			SocketSubsystem->GetSocketError(Result.Error),
 			*Describe());
@@ -303,6 +354,35 @@ void UIpConnection::HandleSocketSendResult(const FSocketSendResult& Result, ISoc
 	}
 }
 
+void UIpConnection::HandleSocketRecvError(class UNetDriver* NetDriver, const FString& ErrorString)
+{
+	check(NetDriver);
+
+	if (SocketErrorDisconnectDelay > 0.f)
+	{
+		const float Time = NetDriver->Time;
+		if (SocketError_RecvDelayStartTime == 0.f)
+		{
+			UE_LOG(LogNet, Log, TEXT("%s. %s Connection beginning close timeout (Timeout = %d)."),
+				*ErrorString,
+				*Describe(),
+				SocketErrorDisconnectDelay);
+
+			SocketError_RecvDelayStartTime = Time;
+			return;
+		}
+		else if ((Time - SocketError_RecvDelayStartTime) < SocketErrorDisconnectDelay)
+		{
+			return;
+		}
+	}
+
+	// For now, this is only called on clients when the ServerConnection fails.
+	// Because of that, on failure we'll shut down the NetDriver.
+	GEngine->BroadcastNetworkFailure(NetDriver->GetWorld(), NetDriver, ENetworkFailure::ConnectionLost, ErrorString);
+	NetDriver->Shutdown();
+}
+
 FString UIpConnection::LowLevelGetRemoteAddress(bool bAppendPort)
 {
 	return RemoteAddr->ToString(bAppendPort);
diff --git a/Engine/Plugins/Online/OnlineSubsystemUtils/Source/OnlineSubsystemUtils/Private/IpNetDriver.cpp b/Engine/Plugins/Online/OnlineSubsystemUtils/Source/OnlineSubsystemUtils/Private/IpNetDriver.cpp
--- a/Engine/Plugins/Online/OnlineSubsystemUtils/Source/OnlineSubsystemUtils/Private/IpNetDriver.cpp
+++ b/Engine/Plugins/Online/OnlineSubsystemUtils/Source/OnlineSubsystemUtils/Private/IpNetDriver.cpp
@@ -507,6 +507,16 @@ bool UIpNetDriver::InitListen( FNetworkNotify* InNotify, FURL& LocalURL, bool bR
 	return true;
 }
 
+class FIpConnectionHelper
+{
+private:
+	friend class UIpNetDriver;
+	static void HandleSocketRecvError(UIpNetDriver* Driver, UIpConnection* Connection, const FString& ErrorString)
+	{
+		Connection->HandleSocketRecvError(Driver, ErrorString);
+	}
+};
+
 void UIpNetDriver::TickDispatch(float DeltaTime)
 {
 	LLM_SCOPE(ELLMTag::Networking);
@@ -623,13 +633,19 @@ void UIpNetDriver::TickDispatch(float DeltaTime)
 				// This should only occur on clients - on servers it leaves the NetDriver in an invalid/vulnerable state
 				if (MyServerConnection != nullptr)
 				{
-					GEngine->BroadcastNetworkFailure(GetWorld(), this, ENetworkFailure::ConnectionLost, ErrorString);
-					Shutdown();
+					// TODO: Maybe we should check to see whether or not the From address matches the server?
+					// If not, we could forward errors incorrectly, causing the connection to shut down.
 
+					FIpConnectionHelper::HandleSocketRecvError(this, MyServerConnection, ErrorString);
 					break;
 				}
 				else
 				{
+					// TODO: Should we also forward errors to connections here?
+					// If we did, instead of just shutting down the NetDriver completely we could instead
+					// boot the given connection.
+					// May be DDoS concerns with the cost of looking up the connections for malicious packets
+					// from sources that won't have connections.
 					UE_CLOG(!DDoS.CheckLogRestrictions(), LogNet, Warning, TEXT("%s"), *ErrorString);
 				}
 
