Add more documentation covering the Networking / Game connection and handshaking processes.
Fix up a few logs to include more information during handshake retries.

diff --git a/Engine/Plugins/Online/OnlineSubsystemUtils/Source/OnlineSubsystemUtils/Private/IpNetDriver.cpp b/Engine/Plugins/Online/OnlineSubsystemUtils/Source/OnlineSubsystemUtils/Private/IpNetDriver.cpp
--- a/Engine/Plugins/Online/OnlineSubsystemUtils/Source/OnlineSubsystemUtils/Private/IpNetDriver.cpp
+++ b/Engine/Plugins/Online/OnlineSubsystemUtils/Source/OnlineSubsystemUtils/Private/IpNetDriver.cpp
@@ -25,6 +25,98 @@ Notes:
 #include "IPAddress.h"
 #include "Sockets.h"
 
+/**
+ * Initiating Connections / Handshaking Flow.
+ *
+ * Both Server and Clients will have have their own NetDrivers, and all UE Replicated Game traffic will be sent by or Received
+ * from the IpNetDriver. This traffic also includes logic for establishing connections, and re-establishing connections when
+ * something goes wrong.
+ *
+ * Handshaking is split across a couple of different places: NetDriver, PendingNetGame, World, PacketHandlers, and maybe others.
+ * The split is due to having separate needs, things such as: determining whether or not an incoming connection is sending data in "UE-Protocol",
+ * determining whether or not an address appears to be malicious, whether or not a given client has the correct version of a game, etc.
+ *
+ *****************************************************************************************
+ * IpNetDriver / IpConnection Startup and Handshaking
+ *****************************************************************************************
+ *
+ * Whenenever a server Loads a map (via UEngine::LoadMap), we will make a call into UWorld::Listen.
+ * That code is responsible for creating the main Game Net Driver, parsing out settings, and calling UNetDriver::InitListen.
+ * Ultimately, that code will be responsible for figuring out what IP / Port we bind to by calls to our configured Socket Subsystem
+ * (see ISocketSubsystem::GetLocalBindAddr and ISocketSubsystem::BindNextPort).
+ *
+ * Once the server is listening, it's ready to start accepting client connections.
+ *
+ * Whenever a client wants to Join a server, they will first establish a new UPendingNetGame in UEngine::Browse with the server's IP.
+ * UPendingNetGame::Initialize and UPendingNetGame::InitNetDriver are responsible for initializing settings and setting up the NetDriver respectively.
+ * Clients will immediately setup a UNetConnection for the server as a part of this initialization, and will start sending data to the server on that connection,
+ * initiating the handshaking process.
+ *
+ * On both Clients and Server, UIpNetDriver::TickDispatch is responsible for receiving packets.
+ * Whenever we receive a packet, we inspect its address and see whether or not it's from a connection we already know about.
+ * We determine whether or not we've established a connection for a given source address by simply keeping a map from FInternetAddr to UNetConnection.
+ *
+ * If a packet is from a connection that's already established, we pass the packet along to the connection via UNetConnection::ReceivedRawPacket.
+ * If a packet is not from a connection that's already established, we treat is as "connectionless" and begin the handshaking process.
+ *
+ * See StatelessConnectionHandlerComponent.cpp for details on how this handshaking works.
+ *
+ *****************************************************************************************
+ * UWorld / UPendingNetGame / AGameModeBase Startup and Handshaking
+ *****************************************************************************************
+ *
+ * After the UIpNetDriver and UIpNetConnection have completed their handshaking process on Client and Server,
+ * UPendingNetGame::SendInitialJoin will be called on the Client to kick off game level handshaking.
+ *
+ * Game Level Handshaking is done through a more structured and involved set of FNetControlMessages.
+ * The full set of control messages can be found in DataChannel.h.
+ *
+ * Most of the work for handling these control messages are done either in UWorld::NotifyControlMessage,
+ * and UPendingNetGame::NotifyControlMessage. Briefly, the flow looks like this:
+ *
+ * Client's UPendingNetGame::SendInitialJoin sends NMT_Hello.
+ *
+ * Server's UWorld::NotifyControlMessage receives NMT_Hello, sends NMT_Challenge.
+ *
+ * Client's UPendingNetGame::NotifyControlMessage receives NMT_Challenge, and sends back data in NMT_Login.
+ *
+ * Server's UWorld::NotifyControlMessage receives NMT_Login, verifies challenge data, and then calls AGameModeBase::PreLogin.
+ * If PreLogin doesn't report any errors, Server calls UWorld::WelcomePlayer, which call AGameModeBase::GameWelcomePlayer,
+ * and send NMT_Welcome with map information.
+ *
+ * Client's UPendingNetGame::NotifyControlMessage receives NMT_Welcome, reads the map info (so it can start loading later),
+ * and sends an NMT_NetSpeed message with the configured Net Speed of the client.
+ *
+ * Server's UWorld::NotifyControlMessage receives NMT_NetSpeed, and adjusts the connections Net Speed appropriately.
+ *
+ * At this point, the handshaking is considered to be complete, and the player is fully connected to the game.
+ * Depending on how long it takes to load the map, the client could still receive some non-handshake control messages
+ * on UPendingNetGame before control transitions to UWorld.
+ *
+ * There are also additional steps for handling Encryption when desired.
+ *
+ *****************************************************************************************
+ * Reestablishing Lost Connections
+ *****************************************************************************************
+ *
+ * Throughout the course of a game, it's possible for connections to be lost for a number of reasons.
+ * Internet could drop out, users could switch from LTE to WIFI, they could leave a game, etc.
+ *
+ * If the server initiated one of these disconnects, or is otherwise aware of it (due to a timeout or error),
+ * then the disconnect will be handled by closing the UNetConnection and notifying the game.
+ * At that point, it's up to a game to decide whether or not they support Join In Progress or Rejoins.
+ * If the game does support it, we will completely restart the handshaking flow as above.
+ *
+ * If something just briefly interrupts the client's connection, but the server is never made aware,
+ * then the engine / game will typically recover automatically (albeit with some packet loss / lag spike).
+ *
+ * However, if the Client's IP Address or Port change for any reason, but the server isn't aware of this,
+ * then we will begin a recovery process by redoing the low level handshake. In this case, game code
+ * will not be alerted.
+ *
+ * This process is covered in StatlessConnectionHandlerComponent.cpp.
+ */
+
 /** For backwards compatibility with the engine stateless connect code */
 #ifndef STATELESSCONNECT_HAS_RANDOM_SEQUENCE
 	#define STATELESSCONNECT_HAS_RANDOM_SEQUENCE 0
@@ -844,6 +936,8 @@ UNetConnection* UIpNetDriver::ProcessConnectionlessPacket(const TSharedRef<FInte
 						// @todo: There needs to be a proper/standardized copy API for this. Also in IpConnection.cpp
 						bool bIsValid = false;
 
+						const FString OldAddress = RemoteAddrRef->ToString(true);
+
 						RemoteAddrRef->SetIp(*Address->ToString(false), bIsValid);
 						RemoteAddrRef->SetPort(Address->GetPort());
 
@@ -866,7 +960,8 @@ UNetConnection* UIpNetDriver::ProcessConnectionlessPacket(const TSharedRef<FInte
 
 						ReturnVal = FoundConn;
 
-						UE_LOG(LogNet, Log, TEXT("Updated IP address for connection %s to: %s"), *FoundConn->Describe(), *IncomingAddress);
+						// We shouldn't need to log IncomingAddress, as the UNetConnection should dump it with it's description.
+						UE_LOG(LogNet, Log, TEXT("Updated IP address for connection. Connection = %s, Old Address = %s"), *FoundConn->Describe(), *OldAddress);
 					}
 					else
 					{
