Added MaxSecondsInReceive and NbPacketsBetweenReceiveTimeTest property to limit the amount of time the Receive tick can process incoming packets.
Set both values and the IpNetDriver will stop processing incoming data when the limit is reached.

[RN] minor Engine/Network

diff --git a/Engine/Plugins/Online/OnlineSubsystemUtils/Source/OnlineSubsystemUtils/Classes/IpNetDriver.h b/Engine/Plugins/Online/OnlineSubsystemUtils/Source/OnlineSubsystemUtils/Classes/IpNetDriver.h
--- a/Engine/Plugins/Online/OnlineSubsystemUtils/Source/OnlineSubsystemUtils/Classes/IpNetDriver.h
+++ b/Engine/Plugins/Online/OnlineSubsystemUtils/Source/OnlineSubsystemUtils/Classes/IpNetDriver.h
@@ -144,6 +144,14 @@ private:
 	UPROPERTY(Config)
 	uint32 ClientDesiredSocketSendBufferBytes;
 
+	/** Maximum time in seconds the TickDispatch can loop on received socket data*/
+	UPROPERTY(Config)
+	double MaxSecondsInReceive = 0.0;
+
+	/** Nb of packets to wait before testing if the receive time went over MaxSecondsInReceive */
+	UPROPERTY(Config)
+	int32 NbPacketsBetweenReceiveTimeTest = 0;
+
 	/** Represents a packet received and/or error encountered by the receive thread, if enabled, queued for the game thread to process. */
 	struct FReceivedPacket
 	{
diff --git a/Engine/Plugins/Online/OnlineSubsystemUtils/Source/OnlineSubsystemUtils/Private/IpNetDriver.cpp b/Engine/Plugins/Online/OnlineSubsystemUtils/Source/OnlineSubsystemUtils/Private/IpNetDriver.cpp
--- a/Engine/Plugins/Online/OnlineSubsystemUtils/Source/OnlineSubsystemUtils/Private/IpNetDriver.cpp
+++ b/Engine/Plugins/Online/OnlineSubsystemUtils/Source/OnlineSubsystemUtils/Private/IpNetDriver.cpp
@@ -548,9 +548,14 @@ void UIpNetDriver::TickDispatch(float DeltaTime)
 	double AlarmTime = StartReceiveTime + GIpNetDriverMaxDesiredTimeSliceBeforeAlarmSecs;
 	const bool bSlowFrameChecks = OnNetworkProcessingCausingSlowFrame.IsBound();
 
+	const bool bCheckReceiveTime = (MaxSecondsInReceive > 0.0) && (NbPacketsBetweenReceiveTimeTest > 0);
+	const double BailOutTime = StartReceiveTime + MaxSecondsInReceive;
+	int32 PacketsLeftUntilTimeTest = NbPacketsBetweenReceiveTimeTest;
+
+	bool bContinueProcessing(true);
 
 	// Process all incoming packets
-	for (FPacketIterator It(this); It; ++It)
+	for (FPacketIterator It(this); It && bContinueProcessing; ++It)
 	{
 		// @todo: Remove the slow frame checks, eventually - potential DDoS and Switch platform constraint
 		if (bSlowFrameChecks)
@@ -564,6 +569,22 @@ void UIpNetDriver::TickDispatch(float DeltaTime)
 			}
 		}
 
+		if (bCheckReceiveTime)
+		{
+			--PacketsLeftUntilTimeTest;
+			if (PacketsLeftUntilTimeTest <= 0)
+			{
+				PacketsLeftUntilTimeTest = NbPacketsBetweenReceiveTimeTest;
+
+				const double CurrentTime = FPlatformTime::Seconds();
+				if (CurrentTime > BailOutTime)
+				{
+					bContinueProcessing = false;
+					UE_LOG(LogNet, Warning, TEXT("UIpNetDriver::TickDispatch: Stopping packet reception after processing for more than %f seconds. %s"), MaxSecondsInReceive, *GetName());
+				}
+			}
+		}
+
 		uint8* Data = nullptr;
 		int32 BytesRead = 0;
 		ESocketErrors Error = SE_NO_ERROR;
