Fixed Switch hang due to unexpected networrk error type.

diff --git a/Engine/Plugins/Online/OnlineSubsystemUtils/Source/OnlineSubsystemUtils/Classes/IpNetDriver.h b/Engine/Plugins/Online/OnlineSubsystemUtils/Source/OnlineSubsystemUtils/Classes/IpNetDriver.h
--- a/Engine/Plugins/Online/OnlineSubsystemUtils/Source/OnlineSubsystemUtils/Classes/IpNetDriver.h
+++ b/Engine/Plugins/Online/OnlineSubsystemUtils/Source/OnlineSubsystemUtils/Classes/IpNetDriver.h
@@ -114,8 +114,8 @@ private:
 	 */
 	static FORCEINLINE bool IsRecvFailBlocking(ESocketErrors Error)
 	{
-		// SE_ECONNABORTED is for PS4 LAN cable pulls
-		return Error == SE_NO_ERROR || Error == SE_EWOULDBLOCK || Error == SE_ECONNABORTED;
+		// SE_ECONNABORTED is for PS4 LAN cable pulls, SE_ENETDOWN is for a Switch hang
+		return Error == SE_NO_ERROR || Error == SE_EWOULDBLOCK || Error == SE_ECONNABORTED || Error == SE_ENETDOWN;
 	};
 
 public:
