Ensure players are removed from party beacon when they leave a party.
Added functionality to remove members from the party beacon instead of just adding.
Increased messaging detail when JIP fails.

diff --git a/Engine/Plugins/Online/OnlineFramework/Source/Party/Private/Party/SocialParty.cpp b/Engine/Plugins/Online/OnlineFramework/Source/Party/Private/Party/SocialParty.cpp
--- a/Engine/Plugins/Online/OnlineFramework/Source/Party/Private/Party/SocialParty.cpp
+++ b/Engine/Plugins/Online/OnlineFramework/Source/Party/Private/Party/SocialParty.cpp
@@ -626,6 +626,39 @@ void USocialParty::HandlePartyJoinRequestReceived(const FUniqueNetId& LocalUserI
 	}
 }
 
+void USocialParty::RemovePlayerFromReservationBeacon(const FUniqueNetId& LocalUserId, const FUniqueNetId& PlayerToRemove)
+{
+	if (!IsLocalPlayerPartyLeader())
+	{
+		return;
+	}
+
+	IOnlinePartyPtr PartyInterface = Online::GetPartyInterfaceChecked(GetWorld());
+
+	FPartyJoinApproval JoinApproval;
+	JoinApproval.SetApprovalAction(EApprovalAction::EnqueueAndStartBeacon);
+	JoinApproval.SetDenialReason(EPartyJoinDenialReason::NoReason);
+
+	FUserPlatform MemberPlatform = FUserPlatform();
+
+	// Enqueue for a more opportune time
+	UE_LOG(LogParty, Verbose, TEXT("[%s] Enqueuing ReservationBeacon update to remove player %s"), *GetPartyId().ToString(), *PlayerToRemove.ToDebugString());
+
+	FPendingMemberApproval PendingApproval;
+	PendingApproval.RecipientId.SetUniqueNetId(LocalUserId.AsShared());
+	PendingApproval.SenderId.SetUniqueNetId(PlayerToRemove.AsShared());
+	PendingApproval.Platform = MemberPlatform;
+	PendingApproval.bIsJIPApproval = true;
+	PendingApproval.bIsPlayerRemoval = true;
+	PendingApprovals.Enqueue(PendingApproval);
+
+
+	if (!ReservationBeaconClient)
+	{
+		ConnectToReservationBeacon();
+	}
+}
+
 void USocialParty::HandlePartyJIPRequestReceived(const FUniqueNetId& LocalUserId, const FOnlinePartyId& PartyId, const FUniqueNetId& SenderId)
 {
 	if (!IsLocalPlayerPartyLeader() || PartyId != GetPartyId())
@@ -938,6 +971,8 @@ void USocialParty::HandlePartyMemberExited(const FUniqueNetId& LocalUserId, cons
 				LeftMember.NotifyRemovedFromParty(ExitReason);
 				LeftMember.MarkPendingKill();
 
+				RemovePlayerFromReservationBeacon(LocalUserId, MemberId);
+
 				// Update party join state, will cause a failure on leader promotion currently
 				// because we can't tell the difference between "expected leader" and "actually the new leader"
 				RefreshPublicJoinability();
@@ -1257,12 +1292,12 @@ void USocialParty::ConnectToReservationBeacon()
 							}
 							else
 							{
-								Reservation.bAllowCrossplay = true; // THis will not matter since we are JIP, and the session already has crossplay set.
+								Reservation.bAllowCrossplay = true; // This will not matter since we are JIP, and the session already has crossplay set.
 							}
 
 							TArray<FPlayerReservation> ReservationAsArray;
 							ReservationAsArray.Add(Reservation);
-							bStartedConnection = ReservationBeaconClient->RequestReservationUpdate(URL, Session->GetSessionIdStr(), GetPartyLeader()->GetPrimaryNetId(), ReservationAsArray);
+							bStartedConnection = ReservationBeaconClient->RequestReservationUpdate(URL, Session->GetSessionIdStr(), GetPartyLeader()->GetPrimaryNetId(), ReservationAsArray, NextApproval.bIsPlayerRemoval);
 						}
 					}
 				}
diff --git a/Engine/Plugins/Online/OnlineFramework/Source/Party/Public/Party/SocialParty.h b/Engine/Plugins/Online/OnlineFramework/Source/Party/Public/Party/SocialParty.h
--- a/Engine/Plugins/Online/OnlineFramework/Source/Party/Public/Party/SocialParty.h
+++ b/Engine/Plugins/Online/OnlineFramework/Source/Party/Public/Party/SocialParty.h
@@ -329,6 +329,7 @@ private:	// Handlers
 
 	void HandleLeavePartyComplete(const FUniqueNetId& LocalUserId, const FOnlinePartyId& PartyId, ELeavePartyCompletionResult LeaveResult, FOnLeavePartyAttemptComplete OnAttemptComplete);
 
+	void RemovePlayerFromReservationBeacon(const FUniqueNetId& LocalUserId, const FUniqueNetId& PlayerToRemove);
 private:
 	TSharedPtr<const FOnlineParty> OssParty;
 
@@ -354,6 +355,7 @@ private:
 		FUniqueNetIdRepl SenderId;
 		FUserPlatform Platform;
 		bool bIsJIPApproval;
+		bool bIsPlayerRemoval = false;
 		TSharedPtr<const FOnlinePartyData> JoinData;
 	};
 	TQueue<FPendingMemberApproval> PendingApprovals;
diff --git a/Engine/Plugins/Online/OnlineSubsystemUtils/Source/OnlineSubsystemUtils/Private/PartyBeaconClient.cpp b/Engine/Plugins/Online/OnlineSubsystemUtils/Source/OnlineSubsystemUtils/Private/PartyBeaconClient.cpp
--- a/Engine/Plugins/Online/OnlineSubsystemUtils/Source/OnlineSubsystemUtils/Private/PartyBeaconClient.cpp
+++ b/Engine/Plugins/Online/OnlineSubsystemUtils/Source/OnlineSubsystemUtils/Private/PartyBeaconClient.cpp
@@ -177,16 +177,16 @@ bool APartyBeaconClient::RequestReservation(const FOnlineSessionSearchResult& De
 	return bSuccess;
 }
 
-bool APartyBeaconClient::RequestReservationUpdate(const FUniqueNetIdRepl& RequestingPartyLeader, const TArray<FPlayerReservation>& PlayersToAdd)
+bool APartyBeaconClient::RequestReservationUpdate(const FUniqueNetIdRepl& RequestingPartyLeader, const TArray<FPlayerReservation>& PlayersToModify, bool bRemovePlayers)
 {
 	bool bWasStarted = false;
 
 	EBeaconConnectionState MyConnectionState = GetConnectionState();
 	if (ensure(MyConnectionState == EBeaconConnectionState::Open))
 	{
-		RequestType = EClientRequestType::ReservationUpdate;
+		RequestType = bRemovePlayers ? EClientRequestType::ReservationRemoveMembers : EClientRequestType::ReservationUpdate;
 		PendingReservation.PartyLeader = RequestingPartyLeader;
-		PendingReservation.PartyMembers = PlayersToAdd;
+		PendingReservation.PartyMembers = PlayersToModify;
 		ServerUpdateReservationRequest(DestSessionId, PendingReservation);
 		bPendingReservationSent = true;
 		bWasStarted = true;
@@ -195,7 +195,7 @@ bool APartyBeaconClient::RequestReservationUpdate(const FUniqueNetIdRepl& Reques
 	return bWasStarted;
 }
 
-bool APartyBeaconClient::RequestReservationUpdate(const FString& ConnectInfoStr, const FString& InSessionId, const FUniqueNetIdRepl& RequestingPartyLeader, const TArray<FPlayerReservation>& PlayersToAdd)
+bool APartyBeaconClient::RequestReservationUpdate(const FString& ConnectInfoStr, const FString& InSessionId, const FUniqueNetIdRepl& RequestingPartyLeader, const TArray<FPlayerReservation>& PlayersToModify, bool bRemovePlayers)
 {
 	bool bWasStarted = false;
 
@@ -205,16 +205,16 @@ bool APartyBeaconClient::RequestReservationUpdate(const FString& ConnectInfoStr,
 		if (MyConnectionState != EBeaconConnectionState::Open)
 		{
 			// create a new pending reservation for these players in the same way as a new reservation request
-			bWasStarted = RequestReservation(ConnectInfoStr, InSessionId, RequestingPartyLeader, PlayersToAdd);
+			bWasStarted = RequestReservation(ConnectInfoStr, InSessionId, RequestingPartyLeader, PlayersToModify);
 			if (bWasStarted)
 			{
 				// Treat this reservation as an update to an existing reservation on the host
-				RequestType = EClientRequestType::ReservationUpdate;
+				RequestType = bRemovePlayers ? EClientRequestType::ReservationRemoveMembers : EClientRequestType::ReservationUpdate;
 			}
 		}
 		else
 		{
-			bWasStarted = RequestReservationUpdate(RequestingPartyLeader, PlayersToAdd);
+			bWasStarted = RequestReservationUpdate(RequestingPartyLeader, PlayersToModify, bRemovePlayers);
 		}
 	}
 	else
@@ -225,7 +225,7 @@ bool APartyBeaconClient::RequestReservationUpdate(const FString& ConnectInfoStr,
 	return bWasStarted;
 }
 
-bool APartyBeaconClient::RequestReservationUpdate(const FOnlineSessionSearchResult& DesiredHost, const FUniqueNetIdRepl& RequestingPartyLeader, const TArray<FPlayerReservation>& PlayersToAdd)
+bool APartyBeaconClient::RequestReservationUpdate(const FOnlineSessionSearchResult& DesiredHost, const FUniqueNetIdRepl& RequestingPartyLeader, const TArray<FPlayerReservation>& PlayersToModify, bool bRemovePlayers)
 {
 	bool bWasStarted = false;
 
@@ -233,16 +233,16 @@ bool APartyBeaconClient::RequestReservationUpdate(const FOnlineSessionSearchResu
 	if (MyConnectionState != EBeaconConnectionState::Open)
 	{
 		// create a new pending reservation for these players in the same way as a new reservation request
-		bWasStarted = RequestReservation(DesiredHost, RequestingPartyLeader, PlayersToAdd);
+		bWasStarted = RequestReservation(DesiredHost, RequestingPartyLeader, PlayersToModify);
 		if (bWasStarted)
 		{
 			// Treat this reservation as an update to an existing reservation on the host
-			RequestType = EClientRequestType::ReservationUpdate;
+			RequestType = bRemovePlayers ? EClientRequestType::ReservationRemoveMembers : EClientRequestType::ReservationUpdate;
 		}
 	}
 	else
 	{
-		RequestReservationUpdate(RequestingPartyLeader, PlayersToAdd);
+		RequestReservationUpdate(RequestingPartyLeader, PlayersToModify, bRemovePlayers);
 	}
 
 	return bWasStarted;
@@ -299,6 +299,12 @@ void APartyBeaconClient::OnConnected()
 			ServerUpdateReservationRequest(DestSessionId, PendingReservation);
 			bPendingReservationSent = true;
 		}
+		else if (RequestType == EClientRequestType::ReservationRemoveMembers)
+		{
+			UE_LOG(LogPartyBeacon, Verbose, TEXT("Party beacon connection established, sending reservation update request to remove player."));
+			ServerRemoveMemberFromReservationRequest(DestSessionId, PendingReservation);
+			bPendingReservationSent = true;
+		}
 		else
 		{
 			UE_LOG(LogPartyBeacon, Warning, TEXT("Failed to handle reservation request type %s"), ToString(RequestType));
@@ -361,7 +367,23 @@ void APartyBeaconClient::ServerUpdateReservationRequest_Implementation(const FSt
 	{
 		PendingReservation = ReservationUpdate;
 		RequestType = EClientRequestType::ReservationUpdate;
-		BeaconHost->ProcessReservationUpdateRequest(this, SessionId, ReservationUpdate);
+		BeaconHost->ProcessReservationUpdateRequest(this, SessionId, ReservationUpdate, false);
+	}
+}
+
+bool APartyBeaconClient::ServerRemoveMemberFromReservationRequest_Validate(const FString& SessionId, const FPartyReservation& ReservationUpdate)
+{
+	return !SessionId.IsEmpty() && ReservationUpdate.PartyLeader.IsValid() && ReservationUpdate.PartyMembers.Num() > 0;
+}
+
+void APartyBeaconClient::ServerRemoveMemberFromReservationRequest_Implementation(const FString& SessionId, const FPartyReservation& ReservationUpdate)
+{
+	APartyBeaconHost* BeaconHost = Cast<APartyBeaconHost>(GetBeaconOwner());
+	if (BeaconHost)
+	{
+		PendingReservation = ReservationUpdate;
+		RequestType = EClientRequestType::ReservationRemoveMembers;
+		BeaconHost->ProcessReservationUpdateRequest(this, SessionId, ReservationUpdate, true);
 	}
 }
 
diff --git a/Engine/Plugins/Online/OnlineSubsystemUtils/Source/OnlineSubsystemUtils/Private/PartyBeaconHost.cpp b/Engine/Plugins/Online/OnlineSubsystemUtils/Source/OnlineSubsystemUtils/Private/PartyBeaconHost.cpp
--- a/Engine/Plugins/Online/OnlineSubsystemUtils/Source/OnlineSubsystemUtils/Private/PartyBeaconHost.cpp
+++ b/Engine/Plugins/Online/OnlineSubsystemUtils/Source/OnlineSubsystemUtils/Private/PartyBeaconHost.cpp
@@ -330,6 +330,31 @@ void APartyBeaconHost::SendReservationUpdates()
 	}
 }
 
+void APartyBeaconHost::PlayerRemoved(const FPlayerReservation& RemovedPlayer)
+{
+	if (RemovedPlayer.UniqueId.IsValid())
+	{
+		if (State)
+		{
+			FUniqueNetIdMatcher PlayerMatch(*RemovedPlayer.UniqueId);
+			int32 FoundIdx = State->PlayersPendingJoin.IndexOfByPredicate(PlayerMatch);
+			if (FoundIdx != INDEX_NONE)
+			{
+				UE_LOG(LogPartyBeacon, Verbose, TEXT("Beacon removing pending player %s"), *RemovedPlayer.UniqueId.ToDebugString());
+				State->PlayersPendingJoin.Remove(RemovedPlayer.UniqueId.GetUniqueNetId());
+			}
+		}
+		else
+		{
+			UE_LOG(LogPartyBeacon, Warning, TEXT("Beacon skipping PlayersPendingJoin for beacon with no state!"));
+		}
+	}
+	else
+	{
+		UE_LOG(LogPartyBeacon, Warning, TEXT("Beacon skipping PlayersPendingJoin for invalid player!"));
+	}
+}
+
 void APartyBeaconHost::NewPlayerAdded(const FPlayerReservation& NewPlayer)
 {
 	if (NewPlayer.UniqueId.IsValid())
@@ -664,7 +689,7 @@ EPartyReservationResult::Type APartyBeaconHost::AddPartyReservation(const FParty
 	return Result;
 }
 
-EPartyReservationResult::Type APartyBeaconHost::UpdatePartyReservation(const FPartyReservation& ReservationUpdateRequest)
+EPartyReservationResult::Type APartyBeaconHost::UpdatePartyReservation(const FPartyReservation& ReservationUpdateRequest, bool bIsRemovingMembers)
 {
 	if (UE_LOG_ACTIVE(LogPartyBeacon, Verbose))
 	{
@@ -678,175 +703,248 @@ EPartyReservationResult::Type APartyBeaconHost::UpdatePartyReservation(const FPa
 		return EPartyReservationResult::ReservationDenied;
 	}
 
-	if (ReservationUpdateRequest.IsValid())
+	if (ReservationUpdateRequest.IsValid() || bIsRemovingMembers)
 	{
-		if (!State->IsBeaconFull())
+		if (bIsRemovingMembers)
 		{
+			UE_LOG(LogPartyBeacon, Verbose, TEXT("Removing Member"));
 			const int32 ExistingReservationIdx = State->GetExistingReservation(ReservationUpdateRequest.PartyLeader);
 			if (ExistingReservationIdx != INDEX_NONE)
 			{
-				// Count the number of available slots for the existing reservation's team
 				TArray<FPartyReservation>& Reservations = State->GetReservations();
 				FPartyReservation& ExistingReservation = Reservations[ExistingReservationIdx];
-				const int32 NumTeamMembers = GetNumPlayersOnTeam(ExistingReservation.TeamNum);
-				const int32 NumAvailableSlotsOnTeam = FMath::Max<int32>(0, GetMaxPlayersPerTeam() - NumTeamMembers);
-				int32 NumPlayersWithExistingReservation = 0;
 
-				// Read the list of new players and remove the ones that have existing reservation entries
-				TArray<FPlayerReservation> NewPlayers;
+				// Read the list of players and remove the ones that are not in this party
+				TArray<FPlayerReservation> PlayersToDelete;
 				for (int32 PlayerIdx = 0; PlayerIdx < ReservationUpdateRequest.PartyMembers.Num(); PlayerIdx++)
 				{
-					const FPlayerReservation& NewPlayerRes = ReservationUpdateRequest.PartyMembers[PlayerIdx];
+					const FPlayerReservation& MarkedPlayerRes = ReservationUpdateRequest.PartyMembers[PlayerIdx];
 
-					const int32 FormerReservationIdx = State->GetExistingReservationContainingMember(NewPlayerRes.UniqueId);
-					if (FormerReservationIdx != ExistingReservationIdx)
+					const int32 FormerReservationIdx = State->GetExistingReservationContainingMember(MarkedPlayerRes.UniqueId);
+					if (FormerReservationIdx != INDEX_NONE && FormerReservationIdx == ExistingReservationIdx)
 					{
-						// player reservation doesn't exist in this reservation so add it as a new player
-						NewPlayers.Add(NewPlayerRes);
-
-						if (FormerReservationIdx != INDEX_NONE)
-						{
-							++NumPlayersWithExistingReservation;
-						}
+						PlayersToDelete.Add(MarkedPlayerRes);
 					}
 					else
 					{
-						// duplicate entry for this player
-						UE_LOG(LogPartyBeacon, Log, TEXT("Skipping player %s because they already have a reservation with this party"),
-							*NewPlayerRes.UniqueId.ToString());
+						// player is not in party
+						UE_LOG(LogPartyBeacon, Log, TEXT("Skipping player %s because they are not in this party"),
+							*MarkedPlayerRes.UniqueId.ToString());
 					}
 				}
 
-				// Validate that adding the new party members to this reservation entry still fits within the team size
-				if ((NewPlayers.Num() - NumPlayersWithExistingReservation) <= NumAvailableSlotsOnTeam)
+				// Copy new player entries into existing reservation
+				for (int32 PlayerIdx = 0; PlayerIdx < PlayersToDelete.Num(); PlayerIdx++)
 				{
-					bool bPlayerRemovedFromReservation = false;
-					if (NewPlayers.Num() > 0)
+					const FPlayerReservation& PlayerRes = PlayersToDelete[PlayerIdx];
+					ExistingReservation.RemoveAllPartyMembers(PlayerRes);
+
+					// Keep track of newly added players
+					PlayerRemoved(PlayerRes);
+					State->SanityCheckReservations(true);
+				}
+				// Update the reservation count before sending the response
+				State->NumConsumedReservations -= PlayersToDelete.Num();
+				UE_LOG(LogPartyBeacon, Verbose, TEXT("APartyBeaconHost::UpdatePartyReservation: Removed %d players, setting NumConsumedReservations to %d"), PlayersToDelete.Num(), State->NumConsumedReservations);
+
+				// Tell any UI and/or clients that there has been a change in the reservation state
+				SendReservationUpdates();
+
+				// Tell the owner that we've received a reservation so the UI can be updated
+				NotifyReservationEventNextFrame(ReservationChanged);
+
+				State->Reservations.RemoveAll([](const FPartyReservation& Reservation)
+				{
+					bool bEmptyReservation = Reservation.PartyMembers.Num() == 0;
+					if (bEmptyReservation)
+					{
+						UE_LOG(LogPartyBeacon, Log, TEXT("Removing reservation with party leader %s because there are no more members in it"),
+							*Reservation.PartyLeader.ToString());
+					}
+					return bEmptyReservation;
+				});
+				State->SanityCheckReservations(false);
+				
+				Result = EPartyReservationResult::ReservationAccepted;
+
+			}
+			else
+			{
+				// Send a not found reservation response
+				Result = EPartyReservationResult::ReservationNotFound;
+			}
+		}
+		else
+		{
+			UE_LOG(LogPartyBeacon, Verbose, TEXT("Adding Member"));
+			if (!State->IsBeaconFull())
+			{
+				const int32 ExistingReservationIdx = State->GetExistingReservation(ReservationUpdateRequest.PartyLeader);
+				if (ExistingReservationIdx != INDEX_NONE)
+				{
+					// Count the number of available slots for the existing reservation's team
+					TArray<FPartyReservation>& Reservations = State->GetReservations();
+					FPartyReservation& ExistingReservation = Reservations[ExistingReservationIdx];
+					const int32 NumTeamMembers = GetNumPlayersOnTeam(ExistingReservation.TeamNum);
+					const int32 NumAvailableSlotsOnTeam = FMath::Max<int32>(0, GetMaxPlayersPerTeam() - NumTeamMembers);
+					int32 NumPlayersWithExistingReservation = 0;
+
+					// Read the list of new players and remove the ones that have existing reservation entries
+					TArray<FPlayerReservation> NewPlayers;
+					for (int32 PlayerIdx = 0; PlayerIdx < ReservationUpdateRequest.PartyMembers.Num(); PlayerIdx++)
 					{
-						if (State->CrossPlayAllowed(ReservationUpdateRequest))
+						const FPlayerReservation& NewPlayerRes = ReservationUpdateRequest.PartyMembers[PlayerIdx];
+
+						const int32 FormerReservationIdx = State->GetExistingReservationContainingMember(NewPlayerRes.UniqueId);
+						if (FormerReservationIdx != ExistingReservationIdx)
 						{
-							// Copy new player entries into existing reservation
-							for (int32 PlayerIdx = 0; PlayerIdx < NewPlayers.Num(); PlayerIdx++)
+							// player reservation doesn't exist in this reservation so add it as a new player
+							NewPlayers.Add(NewPlayerRes);
+
+							if (FormerReservationIdx != INDEX_NONE)
 							{
-								const FPlayerReservation& PlayerRes = NewPlayers[PlayerIdx];
+								++NumPlayersWithExistingReservation;
+							}
+						}
+						else
+						{
+							// duplicate entry for this player
+							UE_LOG(LogPartyBeacon, Log, TEXT("Skipping player %s because they already have a reservation with this party"),
+								*NewPlayerRes.UniqueId.ToString());
+						}
+					}
 
-								// Remove players that existed in other reservations before adding to this reservation
-								if (NumPlayersWithExistingReservation > 0)
+					// Validate that adding the new party members to this reservation entry still fits within the team size
+					if ((NewPlayers.Num() - NumPlayersWithExistingReservation) <= NumAvailableSlotsOnTeam)
+					{
+						bool bPlayerRemovedFromReservation = false;
+						if (NewPlayers.Num() > 0)
+						{
+							if (State->CrossPlayAllowed(ReservationUpdateRequest))
+							{
+								// Copy new player entries into existing reservation
+								for (int32 PlayerIdx = 0; PlayerIdx < NewPlayers.Num(); PlayerIdx++)
 								{
-									const int32 FormerReservationIdx = State->GetExistingReservationContainingMember(PlayerRes.UniqueId);
-									if (FormerReservationIdx != INDEX_NONE)
+									const FPlayerReservation& PlayerRes = NewPlayers[PlayerIdx];
+
+									// Remove players that existed in other reservations before adding to this reservation
+									if (NumPlayersWithExistingReservation > 0)
 									{
-										FPartyReservation& FormerReservation = Reservations[FormerReservationIdx];
-										UE_LOG(LogPartyBeacon, Log, TEXT("APartyBeaconHost::UpdatePartyReservation: Removing player %s from former reservation with leader %s before adding to reservation with leader %s"),
-											*PlayerRes.UniqueId.ToString(), *FormerReservation.PartyLeader.ToString(), *ReservationUpdateRequest.PartyLeader.ToString());
-										if (UE_LOG_ACTIVE(LogPartyBeacon, Verbose))
+										const int32 FormerReservationIdx = State->GetExistingReservationContainingMember(PlayerRes.UniqueId);
+										if (FormerReservationIdx != INDEX_NONE)
 										{
-											FormerReservation.Dump();
-										}
-										int32 NumReservationsRemoved = FormerReservation.RemoveAllPartyMembers(PlayerRes);
+											FPartyReservation& FormerReservation = Reservations[FormerReservationIdx];
+											UE_LOG(LogPartyBeacon, Log, TEXT("APartyBeaconHost::UpdatePartyReservation: Removing player %s from former reservation with leader %s before adding to reservation with leader %s"),
+												*PlayerRes.UniqueId.ToString(), *FormerReservation.PartyLeader.ToString(), *ReservationUpdateRequest.PartyLeader.ToString());
+											if (UE_LOG_ACTIVE(LogPartyBeacon, Verbose))
+											{
+												FormerReservation.Dump();
+											}
+											int32 NumReservationsRemoved = FormerReservation.RemoveAllPartyMembers(PlayerRes);
 
-										State->NumConsumedReservations -= NumReservationsRemoved;
-										UE_LOG(LogPartyBeacon, Verbose, TEXT("APartyBeaconHost::UpdatePartyReservation: Removed %d players, setting NumConsumedReservations to %d"), NumReservationsRemoved, State->NumConsumedReservations);
+											State->NumConsumedReservations -= NumReservationsRemoved;
+											UE_LOG(LogPartyBeacon, Verbose, TEXT("APartyBeaconHost::UpdatePartyReservation: Removed %d players, setting NumConsumedReservations to %d"), NumReservationsRemoved, State->NumConsumedReservations);
 
-										if (NumReservationsRemoved != 0)
-										{
-											bPlayerRemovedFromReservation = true;
-											if (FormerReservation.PartyLeader == PlayerRes.UniqueId)
+											if (NumReservationsRemoved != 0)
 											{
-												UE_LOG(LogPartyBeacon, Display, TEXT("APartyBeaconHost::UpdatePartyReservation: Leader removed, finding member to promote"));
-												// Try to find a new leader for party reservation that lost its leader
-												bool bAnyMemberPromoted = false;
-												for (int32 FormerReservationPlayerIdx = 0; FormerReservationPlayerIdx < FormerReservation.PartyMembers.Num(); FormerReservationPlayerIdx++)
+												bPlayerRemovedFromReservation = true;
+												if (FormerReservation.PartyLeader == PlayerRes.UniqueId)
 												{
-													FPlayerReservation& FormerReservationPlayerEntry = FormerReservation.PartyMembers[FormerReservationPlayerIdx];
-													if (FormerReservationPlayerEntry.UniqueId != FormerReservation.PartyLeader &&
-														FormerReservationPlayerEntry.UniqueId.IsValid() &&
-														State->GetExistingReservation(FormerReservationPlayerEntry.UniqueId) == INDEX_NONE)
+													UE_LOG(LogPartyBeacon, Display, TEXT("APartyBeaconHost::UpdatePartyReservation: Leader removed, finding member to promote"));
+													// Try to find a new leader for party reservation that lost its leader
+													bool bAnyMemberPromoted = false;
+													for (int32 FormerReservationPlayerIdx = 0; FormerReservationPlayerIdx < FormerReservation.PartyMembers.Num(); FormerReservationPlayerIdx++)
 													{
-														// Promote to party leader (for now)
-														UE_LOG(LogPartyBeacon, Display, TEXT("APartyBeaconHost::UpdatePartyReservation: Promoting member %s to leader"), *FormerReservationPlayerEntry.UniqueId.ToDebugString());
-														FormerReservation.PartyLeader = FormerReservationPlayerEntry.UniqueId;
-														bAnyMemberPromoted = true;
-														break;
+														FPlayerReservation& FormerReservationPlayerEntry = FormerReservation.PartyMembers[FormerReservationPlayerIdx];
+														if (FormerReservationPlayerEntry.UniqueId != FormerReservation.PartyLeader &&
+															FormerReservationPlayerEntry.UniqueId.IsValid() &&
+															State->GetExistingReservation(FormerReservationPlayerEntry.UniqueId) == INDEX_NONE)
+														{
+															// Promote to party leader (for now)
+															UE_LOG(LogPartyBeacon, Display, TEXT("APartyBeaconHost::UpdatePartyReservation: Promoting member %s to leader"), *FormerReservationPlayerEntry.UniqueId.ToDebugString());
+															FormerReservation.PartyLeader = FormerReservationPlayerEntry.UniqueId;
+															bAnyMemberPromoted = true;
+															break;
+														}
 													}
+													if (!bAnyMemberPromoted)
+													{
+														UE_LOG(LogPartyBeacon, Display, TEXT("APartyBeaconHost::UpdatePartyReservation: Failed to find a player to promote to leader"));
+													}
+													State->SanityCheckReservations(true);
 												}
-												if (!bAnyMemberPromoted)
-												{
-													UE_LOG(LogPartyBeacon, Display, TEXT("APartyBeaconHost::UpdatePartyReservation: Failed to find a player to promote to leader"));
-												}
-												State->SanityCheckReservations(true);
 											}
 										}
-									}
 
+									}
+									ExistingReservation.PartyMembers.Add(PlayerRes);
+									// Keep track of newly added players
+									NewPlayerAdded(PlayerRes);
+									State->SanityCheckReservations(true);
 								}
-								ExistingReservation.PartyMembers.Add(PlayerRes);
-								// Keep track of newly added players
-								NewPlayerAdded(PlayerRes);
-								State->SanityCheckReservations(true);
-							}
 
-							// Update the reservation count before sending the response
-							State->NumConsumedReservations += NewPlayers.Num();
-							UE_LOG(LogPartyBeacon, Verbose, TEXT("APartyBeaconHost::UpdatePartyReservation: Added %d players, setting NumConsumedReservations to %d"), NewPlayers.Num(), State->NumConsumedReservations);
+								// Update the reservation count before sending the response
+								State->NumConsumedReservations += NewPlayers.Num();
+								UE_LOG(LogPartyBeacon, Verbose, TEXT("APartyBeaconHost::UpdatePartyReservation: Added %d players, setting NumConsumedReservations to %d"), NewPlayers.Num(), State->NumConsumedReservations);
 
-							// Tell any UI and/or clients that there has been a change in the reservation state
-							SendReservationUpdates();
+								// Tell any UI and/or clients that there has been a change in the reservation state
+								SendReservationUpdates();
 
-							// Tell the owner that we've received a reservation so the UI can be updated
-							NotifyReservationEventNextFrame(ReservationChanged);
-							if (State->IsBeaconFull())
-							{
-								// If we've hit our limit, fire the delegate so the host can do the
-								// next step in getting parties together
-								NotifyReservationEventNextFrame(ReservationsFull);
-							}
+								// Tell the owner that we've received a reservation so the UI can be updated
+								NotifyReservationEventNextFrame(ReservationChanged);
+								if (State->IsBeaconFull())
+								{
+									// If we've hit our limit, fire the delegate so the host can do the
+									// next step in getting parties together
+									NotifyReservationEventNextFrame(ReservationsFull);
+								}
 
-							if (bPlayerRemovedFromReservation)
-							{
-								State->Reservations.RemoveAll([](const FPartyReservation& Reservation)
+								if (bPlayerRemovedFromReservation)
 								{
-									bool bEmptyReservation = Reservation.PartyMembers.Num() == 0;
-									if (bEmptyReservation)
+									State->Reservations.RemoveAll([](const FPartyReservation& Reservation)
 									{
-										UE_LOG(LogPartyBeacon, Log, TEXT("Removing reservation with party leader %s because there are no more members in it"),
-											*Reservation.PartyLeader.ToString());
-									}
-									return bEmptyReservation;
-								});
-								State->SanityCheckReservations(false);
-							}
+										bool bEmptyReservation = Reservation.PartyMembers.Num() == 0;
+										if (bEmptyReservation)
+										{
+											UE_LOG(LogPartyBeacon, Log, TEXT("Removing reservation with party leader %s because there are no more members in it"),
+												*Reservation.PartyLeader.ToString());
+										}
+										return bEmptyReservation;
+									});
+									State->SanityCheckReservations(false);
+								}
 
-							Result = EPartyReservationResult::ReservationAccepted;
+								Result = EPartyReservationResult::ReservationAccepted;
+							}
+							else
+							{
+								Result = EPartyReservationResult::ReservationDenied_CrossPlayRestriction;
+							}
 						}
 						else
 						{
-							Result = EPartyReservationResult::ReservationDenied_CrossPlayRestriction;
+							// Duplicate entries (or zero) so existing reservation not updated
+							Result = EPartyReservationResult::ReservationDuplicate;
 						}
 					}
 					else
 					{
-						// Duplicate entries (or zero) so existing reservation not updated
-						Result = EPartyReservationResult::ReservationDuplicate;
+						// Send an invalid party size response
+						Result = EPartyReservationResult::IncorrectPlayerCount;
 					}
 				}
 				else
 				{
-					// Send an invalid party size response
-					Result = EPartyReservationResult::IncorrectPlayerCount;
+					// Send a not found reservation response
+					Result = EPartyReservationResult::ReservationNotFound;
 				}
 			}
 			else
 			{
-				// Send a not found reservation response
-				Result = EPartyReservationResult::ReservationNotFound;
+				// Send a session full response
+				Result = EPartyReservationResult::PartyLimitReached;
 			}
 		}
-		else
-		{
-			// Send a session full response
-			Result = EPartyReservationResult::PartyLimitReached;
-		}
 	}
 	else
 	{
@@ -952,7 +1050,7 @@ void APartyBeaconHost::ProcessReservationRequest(APartyBeaconClient* Client, con
 	}
 }
 
-void APartyBeaconHost::ProcessReservationUpdateRequest(APartyBeaconClient* Client, const FString& SessionId, const FPartyReservation& ReservationUpdateRequest)
+void APartyBeaconHost::ProcessReservationUpdateRequest(APartyBeaconClient* Client, const FString& SessionId, const FPartyReservation& ReservationUpdateRequest, bool bIsRemovingMember)
 {
 	UE_LOG(LogPartyBeacon, Verbose, TEXT("ProcessReservationUpdateRequest %s SessionId %s PartyLeader: %s PartySize: %d from (%s)"),
 		Client ? *Client->GetName() : TEXT("NULL"),
@@ -966,7 +1064,7 @@ void APartyBeaconHost::ProcessReservationUpdateRequest(APartyBeaconClient* Clien
 		EPartyReservationResult::Type Result = EPartyReservationResult::BadSessionId;
 		if (DoesSessionMatch(SessionId))
 		{
-			Result = UpdatePartyReservation(ReservationUpdateRequest);
+			Result = UpdatePartyReservation(ReservationUpdateRequest, bIsRemovingMember);
 		}
 
 		UE_LOG(LogPartyBeacon, Verbose, TEXT("ProcessReservationUpdateRequest result: %s"), EPartyReservationResult::ToString(Result));
diff --git a/Engine/Plugins/Online/OnlineSubsystemUtils/Source/OnlineSubsystemUtils/Public/PartyBeaconClient.h b/Engine/Plugins/Online/OnlineSubsystemUtils/Source/OnlineSubsystemUtils/Public/PartyBeaconClient.h
--- a/Engine/Plugins/Online/OnlineSubsystemUtils/Source/OnlineSubsystemUtils/Public/PartyBeaconClient.h
+++ b/Engine/Plugins/Online/OnlineSubsystemUtils/Source/OnlineSubsystemUtils/Public/PartyBeaconClient.h
@@ -30,7 +30,9 @@ enum class EClientRequestType : uint8
 	/** Simple reconnect (checks for existing reservation) */
 	Reconnect,
 	/** Abandon the reservation beacon (game specific handling)*/
-	Abandon
+	Abandon,
+	/** Remove members from an existing reservation */
+	ReservationRemoveMembers,
 };
 
 inline const TCHAR* ToString(EClientRequestType RequestType)
@@ -57,6 +59,10 @@ inline const TCHAR* ToString(EClientRequestType RequestType)
 	{
 		return TEXT("Reconnect Only");
 	}
+	case EClientRequestType::ReservationRemoveMembers:
+	{
+		return TEXT("Reservation Remove Members");
+	}
 	}
 	return TEXT("");
 }
@@ -131,7 +137,7 @@ class ONLINESUBSYSTEMUTILS_API APartyBeaconClient : public AOnlineBeaconClient
 	 *
 	 * @return true if the request able to be sent, false if it failed to send
 	 */
-	virtual bool RequestReservationUpdate(const FUniqueNetIdRepl& RequestingPartyLeader, const TArray<FPlayerReservation>& PlayersToAdd);
+	virtual bool RequestReservationUpdate(const FUniqueNetIdRepl& RequestingPartyLeader, const TArray<FPlayerReservation>& PlayersToAdd, bool bRemovePlayers = false);
 
 	/**
 	 * Sends an update request to the remote host to append additional members to an existing party
@@ -144,7 +150,7 @@ class ONLINESUBSYSTEMUTILS_API APartyBeaconClient : public AOnlineBeaconClient
 	 *
 	 * @return true if the request able to be sent, false if it failed to send
 	 */
-	virtual bool RequestReservationUpdate(const FString& ConnectInfoStr, const FString& InSessionId, const FUniqueNetIdRepl& RequestingPartyLeader, const TArray<FPlayerReservation>& PlayersToAdd);
+	virtual bool RequestReservationUpdate(const FString& ConnectInfoStr, const FString& InSessionId, const FUniqueNetIdRepl& RequestingPartyLeader, const TArray<FPlayerReservation>& PlayersToAdd, bool bRemovePlayers = false);
 
 	/**
 	 * Sends an update request to the remote host to append additional members to an existing party
@@ -156,7 +162,7 @@ class ONLINESUBSYSTEMUTILS_API APartyBeaconClient : public AOnlineBeaconClient
 	 *
 	 * @return true if the request able to be sent, false if it failed to send
 	 */
-	virtual bool RequestReservationUpdate(const FOnlineSessionSearchResult& DesiredHost, const FUniqueNetIdRepl& RequestingPartyLeader, const TArray<FPlayerReservation>& PlayersToAdd);
+	virtual bool RequestReservationUpdate(const FOnlineSessionSearchResult& DesiredHost, const FUniqueNetIdRepl& RequestingPartyLeader, const TArray<FPlayerReservation>& PlayersToAdd, bool bRemovePlayers = false);
 
 	/**
 	 * Cancel an existing request to the remote host to revoke allocated space on the server.
@@ -285,6 +291,16 @@ protected:
 	UFUNCTION(server, reliable, WithValidation)
 	virtual void ServerUpdateReservationRequest(const FString& SessionId, const FPartyReservation& ReservationUpdate);
 
+	/**
+	 * Tell the server that we are removing members from our reservation
+	 *
+	 * @param SessionId expected session id on the other end (must match)
+	 * @param ReservationUpdate pending reservation request to make with server
+	 */
+	UFUNCTION(server, reliable, WithValidation)
+	virtual void ServerRemoveMemberFromReservationRequest(const FString& SessionId, const FPartyReservation& ReservationUpdate);
+
+
 	/**
 	 * Tell the server to cancel a pending or existing reservation
 	 *
diff --git a/Engine/Plugins/Online/OnlineSubsystemUtils/Source/OnlineSubsystemUtils/Public/PartyBeaconHost.h b/Engine/Plugins/Online/OnlineSubsystemUtils/Source/OnlineSubsystemUtils/Public/PartyBeaconHost.h
--- a/Engine/Plugins/Online/OnlineSubsystemUtils/Source/OnlineSubsystemUtils/Public/PartyBeaconHost.h
+++ b/Engine/Plugins/Online/OnlineSubsystemUtils/Source/OnlineSubsystemUtils/Public/PartyBeaconHost.h
@@ -252,10 +252,11 @@ class ONLINESUBSYSTEMUTILS_API APartyBeaconHost : public AOnlineBeaconHostObject
 	 * An existing reservation for this party leader must already exist
 	 * 
 	 * @param ReservationUpdateRequest reservation update information (doesn't need existing party members, simply a delta)
+	 * @param bIsRemovingMembers set to true if we wish to remove these members from the reservation instead of adding)
 	 *
 	 * @return update attempt result
 	 */
-	virtual EPartyReservationResult::Type UpdatePartyReservation(const FPartyReservation& ReservationUpdateRequest);
+	virtual EPartyReservationResult::Type UpdatePartyReservation(const FPartyReservation& ReservationUpdateRequest, bool bIsRemovingMembers);
 
 	/**
 	 * Attempts to remove a party reservation from the beacon
@@ -297,8 +298,9 @@ class ONLINESUBSYSTEMUTILS_API APartyBeaconHost : public AOnlineBeaconHostObject
 	 * @param Client client beacon making the request
 	 * @param SessionId id of the session that is being checked
 	 * @param ReservationRequest payload of the update request (existing reservation for party leader required)
+	 * @param bRemovingMembers true if this update is removing members, false if adding members.
 	 */
-	virtual void ProcessReservationUpdateRequest(APartyBeaconClient* Client, const FString& SessionId, const FPartyReservation& ReservationUpdateRequest);
+	virtual void ProcessReservationUpdateRequest(APartyBeaconClient* Client, const FString& SessionId, const FPartyReservation& ReservationUpdateRequest, bool bIsRemovingMember);
 
 	/**
 	 * Handle a reservation cancellation request received from an incoming client
@@ -407,6 +409,13 @@ protected:
 	 */
 	void NewPlayerAdded(const FPlayerReservation& NewPlayer);
 
+	/**
+	 * Handle a newly removed player
+	 *
+	 * @param NewPlayer reservation of newly removed player
+	 */
+	void PlayerRemoved(const FPlayerReservation& RemovedPlayer);
+
 	/**
 	 * Does the session match the one associated with this beacon
 	 *
