[NO NON-PII COMMIT MESSAGE]

diff --git a/Engine/Plugins/Online/OnlineSubsystemUtils/Source/OnlineSubsystemUtils/Private/VoiceEngineImpl.cpp b/Engine/Plugins/Online/OnlineSubsystemUtils/Source/OnlineSubsystemUtils/Private/VoiceEngineImpl.cpp
--- a/Engine/Plugins/Online/OnlineSubsystemUtils/Source/OnlineSubsystemUtils/Private/VoiceEngineImpl.cpp
+++ b/Engine/Plugins/Online/OnlineSubsystemUtils/Source/OnlineSubsystemUtils/Private/VoiceEngineImpl.cpp
@@ -258,10 +258,9 @@ bool FVoiceEngineImpl::Init(int32 MaxLocalTalkers, int32 MaxRemoteTalkers)
 		FVoiceModule& VoiceModule = FVoiceModule::Get();
 		if (VoiceModule.IsVoiceEnabled())
 		{
-			VoiceCapture = VoiceModule.CreateVoiceCapture();
 			VoiceEncoder = VoiceModule.CreateVoiceEncoder();
 
-			bSuccess = VoiceCapture.IsValid() && VoiceEncoder.IsValid();
+			bSuccess = VoiceEncoder.IsValid();
 			if (bSuccess)
 			{
 #if PLATFORM_WINDOWS
@@ -347,6 +346,17 @@ uint32 FVoiceEngineImpl::StopLocalVoiceProcessing(uint32 LocalUserNum)
 
 uint32 FVoiceEngineImpl::RegisterLocalTalker(uint32 LocalUserNum)
 {
+	if (!VoiceCapture.IsValid())
+	{
+		VoiceCapture = FVoiceModule::Get().CreateVoiceCapture();
+
+		if (!VoiceCapture.IsValid())
+		{
+			UE_LOG_ONLINE_VOICEENGINE(Error, TEXT("RegisterLocalTalker: Failed to create a Voice Capture Device"));
+			return ONLINE_FAIL;
+		}
+	}
+
 	if (OwningUserIndex == INVALID_INDEX)
 	{
 		OwningUserIndex = LocalUserNum;
@@ -834,4 +844,4 @@ void FVoiceEngineImpl::OnDefaultDeviceChanged()
 	bAudioDeviceChanged = true;
 	TimeDeviceChaned = FPlatformTime::Seconds();
 }
-#endif
\ No newline at end of file
+#endif
