[NO NON-PII COMMIT MESSAGE]

diff --git a/Engine/Plugins/Online/OnlineSubsystemUtils/Source/OnlineSubsystemUtils/Private/VoipListenerSynthComponent.cpp b/Engine/Plugins/Online/OnlineSubsystemUtils/Source/OnlineSubsystemUtils/Private/VoipListenerSynthComponent.cpp
--- a/Engine/Plugins/Online/OnlineSubsystemUtils/Source/OnlineSubsystemUtils/Private/VoipListenerSynthComponent.cpp
+++ b/Engine/Plugins/Online/OnlineSubsystemUtils/Source/OnlineSubsystemUtils/Private/VoipListenerSynthComponent.cpp
@@ -13,6 +13,7 @@ static int16 DebugBuffer[2048];
 bool UVoipListenerSynthComponent::Init(int32& SampleRate)
 {
 	NumChannels = 1;
+	SampleRate = UVOIPStatics::GetVoiceSampleRate();
 	MySampleRate = SampleRate;
 
 #if DEBUG_BUFFERING
