Speculative fix for PS4 LAN cable pull infinite loop.

diff --git a/Engine/Plugins/Online/OnlineSubsystemUtils/Source/OnlineSubsystemUtils/Private/IpNetDriver.cpp b/Engine/Plugins/Online/OnlineSubsystemUtils/Source/OnlineSubsystemUtils/Private/IpNetDriver.cpp
--- a/Engine/Plugins/Online/OnlineSubsystemUtils/Source/OnlineSubsystemUtils/Private/IpNetDriver.cpp
+++ b/Engine/Plugins/Online/OnlineSubsystemUtils/Source/OnlineSubsystemUtils/Private/IpNetDriver.cpp
@@ -334,9 +334,9 @@ void UIpNetDriver::TickDispatch(float DeltaTime)
 				Error = SocketSubsystem->GetLastErrorCode();
 			}
 
-			if (Error == SE_EWOULDBLOCK || Error == SE_NO_ERROR)
+			if (Error == SE_EWOULDBLOCK || Error == SE_NO_ERROR || Error == SE_ECONNABORTED)
 			{
-				// No data or no error?
+				// No data or no error? (SE_ECONNABORTED is for PS4 LAN cable pulls)
 				break;
 			}
 			else if (Error != SE_ECONNRESET && Error != SE_UDP_ERR_PORT_UNREACH)
@@ -782,7 +782,7 @@ uint32 UIpNetDriver::FReceiveThreadRunnable::Run()
 				IncomingPacket.Error = SocketSubsystem->GetLastErrorCode();
 
 				// Pass all other errors back to the Game Thread
-				if (IncomingPacket.Error == SE_EWOULDBLOCK || IncomingPacket.Error == SE_NO_ERROR)
+				if (IncomingPacket.Error == SE_EWOULDBLOCK || IncomingPacket.Error == SE_NO_ERROR || IncomingPacket.Error == SE_ECONNABORTED)
 				{
 					continue;
 				}
