Add creation helper function for SerializeHelper in VoiceEngineImpl

diff --git a/Engine/Plugins/Online/OnlineSubsystemUtils/Source/OnlineSubsystemUtils/Private/VoiceEngineImpl.cpp b/Engine/Plugins/Online/OnlineSubsystemUtils/Source/OnlineSubsystemUtils/Private/VoiceEngineImpl.cpp
--- a/Engine/Plugins/Online/OnlineSubsystemUtils/Source/OnlineSubsystemUtils/Private/VoiceEngineImpl.cpp
+++ b/Engine/Plugins/Online/OnlineSubsystemUtils/Source/OnlineSubsystemUtils/Private/VoiceEngineImpl.cpp
@@ -507,10 +507,7 @@ uint32 FVoiceEngineImpl::SubmitRemoteVoiceData(const FUniqueNetIdWrapper& Remote
 	// Generate a streaming wave audio component for voice playback
 	if (QueuedData.VoipSynthComponent == nullptr || QueuedData.VoipSynthComponent->IsPendingKill())
 	{
-		if (SerializeHelper == nullptr)
-		{
-			SerializeHelper = new FVoiceSerializeHelper(this);
-		}
+		CreateSerializeHelper();
 
 		QueuedData.VoipSynthComponent = CreateVoiceSynthComponent(UVOIPStatics::GetVoiceSampleRate());
 		if (QueuedData.VoipSynthComponent)
@@ -762,4 +759,10 @@ int32 FVoiceEngineImpl::GetMaxVoiceRemainderSize()
 	return MAX_VOICE_REMAINDER_SIZE;
 }
 
-
+void FVoiceEngineImpl::CreateSerializeHelper()
+{
+	if (SerializeHelper == nullptr)
+	{
+		SerializeHelper = new FVoiceSerializeHelper(this);
+	}
+}
\ No newline at end of file
diff --git a/Engine/Plugins/Online/OnlineSubsystemUtils/Source/OnlineSubsystemUtils/Public/VoiceEngineImpl.h b/Engine/Plugins/Online/OnlineSubsystemUtils/Source/OnlineSubsystemUtils/Public/VoiceEngineImpl.h
--- a/Engine/Plugins/Online/OnlineSubsystemUtils/Source/OnlineSubsystemUtils/Public/VoiceEngineImpl.h
+++ b/Engine/Plugins/Online/OnlineSubsystemUtils/Source/OnlineSubsystemUtils/Public/VoiceEngineImpl.h
@@ -298,6 +298,7 @@ protected:
 	TArray<uint8>&             GetDecompressedVoiceBuffer() { return DecompressedVoiceBuffer; }
 	FLocalVoiceData*           GetLocalPlayerVoiceData()    { return PlayerVoiceData; }
 	int32                      GetMaxVoiceRemainderSize();
+	void					   CreateSerializeHelper();
 };
 
 typedef TSharedPtr<FVoiceEngineImpl, ESPMode::ThreadSafe> FVoiceEngineImplPtr;
