Improved OSS logging for module loading and initialization.
Added Utils wrapper to check Online::DoesInstanceExist per-world.

diff --git a/Engine/Plugins/Online/OnlineSubsystem/Source/Private/OnlineSubsystemModule.cpp b/Engine/Plugins/Online/OnlineSubsystem/Source/Private/OnlineSubsystemModule.cpp
--- a/Engine/Plugins/Online/OnlineSubsystem/Source/Private/OnlineSubsystemModule.cpp
+++ b/Engine/Plugins/Online/OnlineSubsystem/Source/Private/OnlineSubsystemModule.cpp
@@ -303,6 +303,8 @@ IOnlineSubsystem* FOnlineSubsystemModule::GetOnlineSubsystem(const FName InSubsy
 
 				if (OSSFactory != nullptr)
 				{
+					UE_LOG_ONLINE(Verbose, TEXT("Creating online subsystem instance for: %s"), *InSubsystemName.ToString());
+						
 					IOnlineSubsystemPtr NewSubsystemInstance = (*OSSFactory)->CreateSubsystem(InstanceName);
 					if (NewSubsystemInstance.IsValid())
 					{
diff --git a/Engine/Plugins/Online/OnlineSubsystemFacebook/Source/Private/OnlineSubsystemFacebookModule.cpp b/Engine/Plugins/Online/OnlineSubsystemFacebook/Source/Private/OnlineSubsystemFacebookModule.cpp
--- a/Engine/Plugins/Online/OnlineSubsystemFacebook/Source/Private/OnlineSubsystemFacebookModule.cpp
+++ b/Engine/Plugins/Online/OnlineSubsystemFacebook/Source/Private/OnlineSubsystemFacebookModule.cpp
@@ -20,6 +20,8 @@ public:
 		FOnlineSubsystemFacebookPtr OnlineSub = MakeShared<FOnlineSubsystemFacebook, ESPMode::ThreadSafe>(InstanceName);
 		if (OnlineSub->IsEnabled())
 		{
+			UE_LOG_ONLINE(Log, TEXT("Facebook API is being initialized."));
+
 			if(!OnlineSub->Init())
 			{
 				UE_LOG_ONLINE(Warning, TEXT("Facebook API failed to initialize!"));
@@ -40,7 +42,7 @@ public:
 
 void FOnlineSubsystemFacebookModule::StartupModule()
 {
-	UE_LOG_ONLINE(Log, TEXT("Facebook Startup!"));
+	UE_LOG_ONLINE(Log, TEXT("Facebook Module Startup!"));
 
 	FacebookFactory = new FOnlineFactoryFacebook();
 
@@ -51,7 +53,7 @@ void FOnlineSubsystemFacebookModule::StartupModule()
 
 void FOnlineSubsystemFacebookModule::ShutdownModule()
 {
-	UE_LOG_ONLINE(Log, TEXT("Facebook Shutdown!"));
+	UE_LOG_ONLINE(Log, TEXT("Facebook Module Shutdown!"));
 
 	FOnlineSubsystemModule& OSS = FModuleManager::GetModuleChecked<FOnlineSubsystemModule>("OnlineSubsystem");
 	OSS.UnregisterPlatformService(FACEBOOK_SUBSYSTEM);
diff --git a/Engine/Plugins/Online/OnlineSubsystemUtils/Source/OnlineSubsystemUtils/Public/OnlineSubsystemUtils.h b/Engine/Plugins/Online/OnlineSubsystemUtils/Source/OnlineSubsystemUtils/Public/OnlineSubsystemUtils.h
--- a/Engine/Plugins/Online/OnlineSubsystemUtils/Source/OnlineSubsystemUtils/Public/OnlineSubsystemUtils.h
+++ b/Engine/Plugins/Online/OnlineSubsystemUtils/Source/OnlineSubsystemUtils/Public/OnlineSubsystemUtils.h
@@ -224,7 +224,7 @@ namespace Online
 	}
 
 	/** 
-	 * Determine if the subsystem for a given interface is already loaded
+	 * Determine if the subsystem module for a given interface is already loaded.  Note that modules do not typically initialize the subsystem when the module is loaded.
 	 *
 	 * @param World the world to use for context
 	 * @param SubsystemName name of the requested online service
@@ -249,7 +249,34 @@ namespace Online
 		return IOnlineSubsystem::IsLoaded(SubsystemName);
 #endif
 	}
-
+	
+	/**
+	 * Determine if the subsystem for a given interface has been instanced
+	 *
+	 * @param World the world to use for context
+	 * @param SubsystemName name of the requested online service
+	 *
+	 * @return true if the subsystem has been instanced
+	 */
+	static bool DoesInstanceExist(const UWorld* World, const FName& SubsystemName = NAME_None)
+	{
+#if UE_EDITOR // at present, multiple worlds are only possible in the editor
+		FName Identifier = SubsystemName;
+		if (World != NULL)
+		{
+			FWorldContext& CurrentContext = GEngine->GetWorldContextFromWorldChecked(World);
+			if (CurrentContext.WorldType == EWorldType::PIE)
+			{
+				Identifier = FName(*FString::Printf(TEXT("%s:%s"), SubsystemName != NAME_None ? *SubsystemName.ToString() : TEXT(""), *CurrentContext.ContextHandle.ToString()));
+			}
+		}
+		
+		return IOnlineSubsystem::DoesInstanceExist(SubsystemName);
+#else
+		return IOnlineSubsystem::DoesInstanceExist(SubsystemName);
+#endif
+	}
+	
 	/** Reimplement all the interfaces of Online.h with support for UWorld accessors */
 	IMPLEMENT_GET_INTERFACE(Session);
 	IMPLEMENT_GET_INTERFACE(Party);
