Speculative fix for PS4 LAN cable pull infinite loop.

diff --git a/Engine/Plugins/Online/OnlineSubsystemUtils/Source/OnlineSubsystemUtils/Private/IpNetDriver.cpp b/Engine/Plugins/Online/OnlineSubsystemUtils/Source/OnlineSubsystemUtils/Private/IpNetDriver.cpp
--- a/Engine/Plugins/Online/OnlineSubsystemUtils/Source/OnlineSubsystemUtils/Private/IpNetDriver.cpp
+++ b/Engine/Plugins/Online/OnlineSubsystemUtils/Source/OnlineSubsystemUtils/Private/IpNetDriver.cpp
@@ -277,9 +277,9 @@ void UIpNetDriver::TickDispatch(float DeltaTime)
 		{
 			ESocketErrors Error = SocketSubsystem->GetLastErrorCode();
 
-			if (Error == SE_EWOULDBLOCK || Error == SE_NO_ERROR)
+			if (Error == SE_EWOULDBLOCK || Error == SE_NO_ERROR || Error == SE_ECONNABORTED)
 			{
-				// No data or no error?
+				// No data or no error? (SE_ECONNABORTED is for PS4 LAN cable pulls)
 				break;
 			}
 			else if (Error != SE_ECONNRESET && Error != SE_UDP_ERR_PORT_UNREACH)
