Broadcaster JIP

diff --git a/Engine/Plugins/Online/OnlineFramework/Source/Party/Private/Party/SocialParty.cpp b/Engine/Plugins/Online/OnlineFramework/Source/Party/Private/Party/SocialParty.cpp
--- a/Engine/Plugins/Online/OnlineFramework/Source/Party/Private/Party/SocialParty.cpp
+++ b/Engine/Plugins/Online/OnlineFramework/Source/Party/Private/Party/SocialParty.cpp
@@ -125,7 +125,8 @@ bool USocialParty::IsJoiningDuringLoadEnabled()
 }
 
 USocialParty::USocialParty()
-	: ReservationBeaconClientClass(APartyBeaconClient::StaticClass())
+	: ReservationBeaconClientClass(APartyBeaconClient::StaticClass()),
+	  SpectatorBeaconClientClass(ASpectatorBeaconClient::StaticClass())
 {}
 
 ECrossplayPreference GetCrossplayPreferenceFromJoinData(const FOnlinePartyData& JoinData)
@@ -1240,6 +1241,18 @@ APartyBeaconClient* USocialParty::CreateReservationBeaconClient()
 	return ReservationBeaconClient;
 }
 
+ASpectatorBeaconClient* USocialParty::CreateSpectatorBeaconClient()
+{
+	UWorld* World = GetWorld();
+	check(World);
+
+	// Clear out our cached net driver name, we're going to create a new one here
+	LastSpectatorBeaconClientNetDriverName = NAME_None;
+	SpectatorBeaconClient = World->SpawnActor<ASpectatorBeaconClient>(SpectatorBeaconClientClass);
+
+	return SpectatorBeaconClient;
+}
+
 void USocialParty::PumpApprovalQueue()
 {
 	// Check if there are any more while we are connected
@@ -1340,6 +1353,20 @@ void USocialParty::CleanupReservationBeacon()
 	}
 }
 
+void USocialParty::CleanupSpectatorBeacon()
+{
+	if (SpectatorBeaconClient)
+	{
+		UE_LOG(LogParty, Verbose, TEXT("Spectator reservation beacon cleanup while in state %s, pending approvals: %s"), ToString(SpectatorBeaconClient->GetConnectionState()), !PendingApprovals.IsEmpty() ? TEXT("true") : TEXT("false"));
+
+		LastReservationBeaconClientNetDriverName = SpectatorBeaconClient->GetNetDriverName();
+		SpectatorBeaconClient->OnHostConnectionFailure().Unbind();
+		SpectatorBeaconClient->OnReservationRequestComplete().Unbind();
+		SpectatorBeaconClient->DestroyBeacon();
+		SpectatorBeaconClient = nullptr;
+	}
+}
+
 FName USocialParty::GetGameSessionName() const
 {
 	const APlayerController* OwnerPC = GetOwningLocalPlayer().GetPlayerController(GetWorld());
diff --git a/Engine/Plugins/Online/OnlineFramework/Source/Party/Public/Party/SocialParty.h b/Engine/Plugins/Online/OnlineFramework/Source/Party/Public/Party/SocialParty.h
--- a/Engine/Plugins/Online/OnlineFramework/Source/Party/Public/Party/SocialParty.h
+++ b/Engine/Plugins/Online/OnlineFramework/Source/Party/Public/Party/SocialParty.h
@@ -6,11 +6,12 @@
 #include "Party/PartyDataReplicator.h"
 
 #include "PartyBeaconState.h"
+#include "SpectatorBeaconState.h"
+#include "SpectatorBeaconClient.h"
 #include "Interfaces/OnlinePartyInterface.h"
 #include "Interfaces/OnlineChatInterface.h"
 #include "Containers/Queue.h"
 #include "Engine/EngineBaseTypes.h"
-
 #include "SocialParty.generated.h"
 
 class APartyBeaconClient;
@@ -242,6 +243,15 @@ protected:
 
 	APartyBeaconClient* GetReservationBeaconClient() const { return ReservationBeaconClient; }
 
+	/**
+	* Create a spectator beacon and connect to the server to get approval for new spectators
+	*/
+	void ConnectToSpectatorBeacon();
+	void CleanupSpectatorBeacon();
+	ASpectatorBeaconClient* CreateSpectatorBeaconClient();
+
+	ASpectatorBeaconClient* GetSpectatorBeaconClient() const { return SpectatorBeaconClient; }
+
 	/** Child classes MUST call EstablishRepDataInstance() on this using their member rep data struct instance */
 	FPartyDataReplicator PartyDataReplicator;
 
@@ -249,6 +259,10 @@ protected:
 	UPROPERTY()
 	TSubclassOf<APartyBeaconClient> ReservationBeaconClientClass;
 
+	/** Spectator beacon class for getting server approval for new spectators while in a game */
+	UPROPERTY()
+		TSubclassOf<ASpectatorBeaconClient> SpectatorBeaconClientClass;
+
 private:
 	UPartyMember* GetOrCreatePartyMember(const FUniqueNetId& MemberId);
 	void PumpApprovalQueue();
@@ -325,10 +339,21 @@ private:
 	 * Some network error handlers may be called after we cleanup our beacon connection.
 	 */
 	FName LastReservationBeaconClientNetDriverName;
-
+	
 	/** Reservation beacon client instance while getting approval for new party members*/
 	UPROPERTY()
-	APartyBeaconClient* ReservationBeaconClient = nullptr;
+		APartyBeaconClient* ReservationBeaconClient = nullptr;
+	
+	/**
+	* Last known spectator beacon client net driver name
+	* Intended to be used to detect network errors related to our current or last spectator beacon client's net driver.
+	* Some network error handlers may be called after we cleanup our beacon connection.
+	*/
+	FName LastSpectatorBeaconClientNetDriverName;
+	
+	/** Spectator beacon client instance while getting approval for spectator*/
+	UPROPERTY()
+		ASpectatorBeaconClient* SpectatorBeaconClient = nullptr;
 
 	/**
 	 * True when we have limited functionality due to lacking an xmpp connection.
diff --git a/Engine/Plugins/Online/OnlineSubsystem/Source/Public/OnlineSessionSettings.h b/Engine/Plugins/Online/OnlineSubsystem/Source/Public/OnlineSessionSettings.h
--- a/Engine/Plugins/Online/OnlineSubsystem/Source/Public/OnlineSessionSettings.h
+++ b/Engine/Plugins/Online/OnlineSubsystem/Source/Public/OnlineSessionSettings.h
@@ -33,6 +33,10 @@
 #define SETTING_NEEDSSORT FName(TEXT("NEEDSSORT"))
 /** Setting describing the session key (value is FString) */
 #define SETTING_SESSIONKEY FName(TEXT("SESSIONKEY"))
+/** Setting describing the session key (value is FString) */
+#define SETTING_ALLOWBROADCASTING FName(TEXT("ALLOWBROADCASTING"))
+/** Setting defining broadcaster (value is bool) */
+#define SETTING_BROADCASTER FName(TEXT("BROADCASTER"))
 
 /** 8 user defined integer params to be used when filtering searches for sessions */
 #define SETTING_CUSTOMSEARCHINT1 FName(TEXT("CUSTOMSEARCHINT1"))
diff --git a/Engine/Plugins/Online/OnlineSubsystemUtils/Source/OnlineSubsystemUtils/Private/SpectatorBeaconClient.cpp b/Engine/Plugins/Online/OnlineSubsystemUtils/Source/OnlineSubsystemUtils/Private/SpectatorBeaconClient.cpp
new file mode 100644
--- /dev/null
+++ b/Engine/Plugins/Online/OnlineSubsystemUtils/Source/OnlineSubsystemUtils/Private/SpectatorBeaconClient.cpp
@@ -0,0 +1,469 @@
+// Copyright 1998-2018 Epic Games, Inc. All Rights Reserved.
+
+#include "SpectatorBeaconClient.h"
+#include "OnlineSubsystemUtils.h"
+#include "SpectatorBeaconHost.h"
+
+#if !UE_BUILD_SHIPPING
+namespace SpectatorBeaconConsoleVariables
+{
+	/** Time to delay delegates firing a reservation request response */
+	TAutoConsoleVariable<float> CVarDelayReservationResponse(
+		TEXT("beacon.DelayReservationResponse"),
+		0.0f,
+		TEXT("Delay time between received response and notification\n")
+		TEXT("Time in secs"),
+		ECVF_Default);
+
+	/** Time to delay delegates firing a cancel reservation request response */
+	TAutoConsoleVariable<float> CVarDelayCancellationResponse(
+		TEXT("beacon.DelayCancellationResponse"),
+		0.0f,
+		TEXT("Delay time between received cancel response and notification\n")
+		TEXT("Time in secs"),
+		ECVF_Default);
+
+	/** Time to delay delegates firing a reservation update response */
+	TAutoConsoleVariable<float> CVarDelayUpdateResponse(
+		TEXT("beacon.DelayUpdateResponse"),
+		0.0f,
+		TEXT("Delay time between received update response and notification\n")
+		TEXT("Time in secs"),
+		ECVF_Default);
+
+	/** Time to delay delegates firing a reservation full response */
+	TAutoConsoleVariable<float> CVarDelayFullResponse(
+		TEXT("beacon.DelayFullResponse"),
+		0.0f,
+		TEXT("Delay time between received full response and notification\n")
+		TEXT("Time in secs"),
+		ECVF_Default);
+}
+#endif
+
+/** Max time to wait for a response from the server for CancelReservation */
+#define CANCEL_FAILSAFE 5.0f
+
+ASpectatorBeaconClient::ASpectatorBeaconClient(const FObjectInitializer& ObjectInitializer) :
+	Super(ObjectInitializer),
+	RequestType(ESpectatorClientRequestType::NonePending),
+	bPendingReservationSent(false),
+	bCancelReservation(false)
+{
+}
+
+void ASpectatorBeaconClient::BeginDestroy()
+{
+	ClearTimers(true);
+	Super::BeginDestroy();
+}
+
+void ASpectatorBeaconClient::ClearTimers(bool bCallFailSafeIfNeeded)
+{
+	if (bCallFailSafeIfNeeded &&
+		CancelRPCFailsafe.IsValid())
+	{
+		UE_LOG(LogBeacon, Verbose, TEXT("Clearing timers with cancel reservation in flight.  Calling Failsafe."));
+
+		UWorld* World = GetWorld();
+		if (World)
+		{
+			FTimerManager& TM = World->GetTimerManager();
+			TM.ClearTimer(CancelRPCFailsafe);
+		}
+		OnCancelledFailsafe();
+	}
+
+	UWorld* World = GetWorld();
+	if (World)
+	{
+		if (PendingResponseTimerHandle.IsValid())
+		{
+			UE_LOG(LogSpectatorBeacon, Verbose, TEXT("ClearTimers: Pending reservation response cleared."));
+		}
+
+		if (PendingCancelResponseTimerHandle.IsValid())
+		{
+			UE_LOG(LogSpectatorBeacon, Verbose, TEXT("ClearTimers: Pending cancel response cleared."));
+		}
+
+		if (PendingReservationUpdateTimerHandle.IsValid())
+		{
+			UE_LOG(LogSpectatorBeacon, Verbose, TEXT("ClearTimers: Pending reservation update cleared."));
+		}
+
+		if (PendingReservationFullTimerHandle.IsValid())
+		{
+			UE_LOG(LogSpectatorBeacon, Verbose, TEXT("ClearTimers: Pending reservation full cleared."));
+		}
+
+		if (CancelRPCFailsafe.IsValid())
+		{
+			UE_LOG(LogSpectatorBeacon, Verbose, TEXT("ClearTimers: Cancel failsafe cleared."));
+		}
+
+		FTimerManager& TM = World->GetTimerManager();
+		TM.ClearTimer(PendingResponseTimerHandle);
+		TM.ClearTimer(PendingCancelResponseTimerHandle);
+		TM.ClearTimer(PendingReservationUpdateTimerHandle);
+		TM.ClearTimer(PendingReservationFullTimerHandle);
+		TM.ClearTimer(CancelRPCFailsafe);
+
+		PendingResponseTimerHandle.Invalidate();
+		PendingCancelResponseTimerHandle.Invalidate();
+		PendingReservationUpdateTimerHandle.Invalidate();
+		PendingReservationFullTimerHandle.Invalidate();
+		CancelRPCFailsafe.Invalidate();
+	}
+}
+
+bool ASpectatorBeaconClient::RequestReservation(const FString& ConnectInfoStr, const FString& InSessionId, const FUniqueNetIdRepl& RequestingSpectator, const FPlayerReservation& Spectator)
+{
+	bool bSuccess = false;
+
+	FURL ConnectURL(NULL, *ConnectInfoStr, TRAVEL_Absolute);
+	UE_LOG(LogSpectatorBeacon, Verbose, TEXT("SpectatorBeaconClient RequestReservation() connectURL: %s"), *ConnectURL.ToString());
+	if (InitClient(ConnectURL))
+	{
+		DestSessionId = InSessionId;
+		PendingReservation.SpectatorId = RequestingSpectator;
+		PendingReservation.Spectator = Spectator;
+		bPendingReservationSent = false;
+		RequestType = ESpectatorClientRequestType::ExistingSessionReservation;
+		bSuccess = true;
+	}
+	else
+	{
+		UE_LOG(LogSpectatorBeacon, Warning, TEXT("RequestReservation: Failure to init client beacon with %s."), *ConnectURL.ToString());
+		RequestType = ESpectatorClientRequestType::NonePending;
+	}
+
+	if (!bSuccess)
+	{
+		OnFailure();
+	}
+
+	return bSuccess;
+}
+
+bool ASpectatorBeaconClient::RequestReservation(const FOnlineSessionSearchResult& DesiredHost, const FUniqueNetIdRepl& RequestingSpectator, const FPlayerReservation& Spectator)
+{
+	bool bSuccess = false;
+
+	if (DesiredHost.IsValid())
+	{
+		UWorld* World = GetWorld();
+
+		IOnlineSubsystem* OnlineSub = Online::GetSubsystem(World);
+		if (OnlineSub)
+		{
+			IOnlineSessionPtr SessionInt = OnlineSub->GetSessionInterface();
+			if (SessionInt.IsValid())
+			{
+				FString ConnectInfo;
+				if (SessionInt->GetResolvedConnectString(DesiredHost, NAME_BeaconPort, ConnectInfo))
+				{
+					FString SessionId = DesiredHost.Session.SessionInfo->GetSessionId().ToString();
+					return RequestReservation(ConnectInfo, SessionId, RequestingSpectator, Spectator);
+				}
+			}
+		}
+	}
+
+	if (!bSuccess)
+	{
+		OnFailure();
+	}
+
+	return bSuccess;
+}
+
+void ASpectatorBeaconClient::CancelReservation()
+{
+	if (ensure(PendingReservation.SpectatorId.IsValid()))
+	{
+		bCancelReservation = true;
+
+		// Clear out any pending response handling, only the cancel matters
+		ClearTimers(false);
+
+		if (bPendingReservationSent)
+		{
+			UE_LOG(LogSpectatorBeacon, Verbose, TEXT("Sending cancel reservation request."));
+			ServerCancelReservationRequest(PendingReservation.SpectatorId);
+
+			// In case the server is loading or unresponsive (ie no host beacon)
+			FTimerDelegate TimerDelegate;
+			TimerDelegate.BindUObject(this, &ThisClass::OnCancelledFailsafe);
+
+			UWorld* World = GetWorld();
+			check(World);
+			World->GetTimerManager().SetTimer(CancelRPCFailsafe, TimerDelegate, CANCEL_FAILSAFE, false);
+		}
+		else
+		{
+			UE_LOG(LogSpectatorBeacon, Verbose, TEXT("Reservation request never sent, no need to send cancelation request."));
+			OnCancelledComplete();
+		}
+	}
+	else
+	{
+		UE_LOG(LogSpectatorBeacon, Verbose, TEXT("Unable to cancel reservation request with invalid spectator."));
+		OnCancelledComplete();
+	}
+}
+
+void ASpectatorBeaconClient::OnConnected()
+{
+	if (!bCancelReservation)
+	{
+		if (RequestType == ESpectatorClientRequestType::ExistingSessionReservation)
+		{
+			UE_LOG(LogSpectatorBeacon, Verbose, TEXT("Spectator beacon connection established, sending join reservation request."));
+			ServerReservationRequest(DestSessionId, PendingReservation);
+			bPendingReservationSent = true;
+		}
+		else
+		{
+			UE_LOG(LogSpectatorBeacon, Warning, TEXT("Failed to handle reservation request type %s"), ToString(RequestType));
+			OnFailure();
+		}
+	}
+	else
+	{
+		UE_LOG(LogSpectatorBeacon, Verbose, TEXT("Reservation request previously canceled, nothing sent."));
+		OnCancelledComplete();
+	}
+}
+
+void ASpectatorBeaconClient::OnCancelledFailsafe()
+{
+	ClientCancelReservationResponse_Implementation(ESpectatorReservationResult::ReservationRequestCanceled);
+}
+
+void ASpectatorBeaconClient::OnCancelledComplete()
+{
+	ReservationRequestComplete.ExecuteIfBound(ESpectatorReservationResult::ReservationRequestCanceled);
+	RequestType = ESpectatorClientRequestType::NonePending;
+	bCancelReservation = false;
+}
+
+void ASpectatorBeaconClient::OnFailure()
+{
+	ClearTimers(true);
+	RequestType = ESpectatorClientRequestType::NonePending;
+	Super::OnFailure();
+}
+
+/// @cond DOXYGEN_WARNINGS
+
+bool ASpectatorBeaconClient::ServerReservationRequest_Validate(const FString& SessionId, const FSpectatorReservation& Reservation)
+{
+	return !SessionId.IsEmpty() && Reservation.SpectatorId.IsValid();
+}
+
+void ASpectatorBeaconClient::ServerReservationRequest_Implementation(const FString& SessionId, const FSpectatorReservation& Reservation)
+{
+	ASpectatorBeaconHost* BeaconHost = Cast<ASpectatorBeaconHost>(GetBeaconOwner());
+	if (BeaconHost)
+	{
+		PendingReservation = Reservation;
+		RequestType = ESpectatorClientRequestType::ExistingSessionReservation;
+		BeaconHost->ProcessReservationRequest(this, SessionId, Reservation);
+	}
+}
+
+bool ASpectatorBeaconClient::ServerCancelReservationRequest_Validate(const FUniqueNetIdRepl& SpectatorId)
+{
+	return true;
+}
+
+void ASpectatorBeaconClient::ServerCancelReservationRequest_Implementation(const FUniqueNetIdRepl& SpectatorId)
+{
+	ASpectatorBeaconHost* BeaconHost = Cast<ASpectatorBeaconHost>(GetBeaconOwner());
+	if (BeaconHost)
+	{
+		bCancelReservation = true;
+		BeaconHost->ProcessCancelReservationRequest(this, SpectatorId);
+	}
+}
+
+void ASpectatorBeaconClient::ClientReservationResponse_Implementation(ESpectatorReservationResult::Type ReservationResponse)
+{
+	if (!bCancelReservation)
+	{
+#if !UE_BUILD_SHIPPING
+		const float Rate = SpectatorBeaconConsoleVariables::CVarDelayReservationResponse.GetValueOnGameThread();
+#else
+		const float Rate = 0.0f;
+#endif
+		if (Rate > 0.0f)
+		{
+			UE_LOG(LogSpectatorBeacon, Verbose, TEXT("Spectator beacon response received %s, waiting %fs to notify"), ESpectatorReservationResult::ToString(ReservationResponse), Rate);
+
+			FTimerDelegate TimerDelegate;
+			TimerDelegate.BindLambda([this, ReservationResponse]()
+			{
+				ProcessReservationResponse(ReservationResponse);
+			});
+
+			PendingResponseTimerHandle = DelayResponse(TimerDelegate, Rate);
+		}
+		else
+		{
+			ProcessReservationResponse(ReservationResponse);
+		}
+	}
+	else
+	{
+		UE_LOG(LogSpectatorBeacon, Verbose, TEXT("Spectator beacon response received %s, ignored due to cancel in progress"), ESpectatorReservationResult::ToString(ReservationResponse));
+		// Cancel RPC or failsafe timer will trigger the cancel
+	}
+}
+
+/// @endcond
+
+void ASpectatorBeaconClient::ProcessReservationResponse(ESpectatorReservationResult::Type ReservationResponse)
+{
+	if (!bCancelReservation)
+	{
+		UE_LOG(LogSpectatorBeacon, Verbose, TEXT("Spectator beacon response received %s"), ESpectatorReservationResult::ToString(ReservationResponse));
+		ReservationRequestComplete.ExecuteIfBound(ReservationResponse);
+		RequestType = ESpectatorClientRequestType::NonePending;
+	}
+	else
+	{
+		UE_LOG(LogSpectatorBeacon, Verbose, TEXT("Spectator beacon response received %s, ignored due to cancel in progress"), ESpectatorReservationResult::ToString(ReservationResponse));
+		// Cancel RPC or failsafe timer will trigger the cancel
+	}
+}
+
+/// @cond DOXYGEN_WARNINGS
+
+void ASpectatorBeaconClient::ClientCancelReservationResponse_Implementation(ESpectatorReservationResult::Type ReservationResponse)
+{
+	ensure(bCancelReservation);
+
+	// Clear out any pending response handling (including failsafe timer)
+	ClearTimers(false);
+#if !UE_BUILD_SHIPPING	
+	const float Rate = SpectatorBeaconConsoleVariables::CVarDelayCancellationResponse.GetValueOnGameThread();
+#else
+	const float Rate = 0.0f;
+#endif
+	if (Rate > 0.0f)
+	{
+		UE_LOG(LogSpectatorBeacon, Verbose, TEXT("Spectator beacon cancellation response received %s, waiting %fs to notify"), ESpectatorReservationResult::ToString(ReservationResponse), Rate);
+
+		FTimerDelegate TimerDelegate;
+		TimerDelegate.BindLambda([this, ReservationResponse]()
+		{
+			ProcessCancelReservationResponse(ReservationResponse);
+		});
+
+		PendingCancelResponseTimerHandle = DelayResponse(TimerDelegate, Rate);
+	}
+	else
+	{
+		ProcessCancelReservationResponse(ReservationResponse);
+	}
+}
+
+/// @endcond
+
+void ASpectatorBeaconClient::ProcessCancelReservationResponse(ESpectatorReservationResult::Type ReservationResponse)
+{
+	ensure(ReservationResponse == ESpectatorReservationResult::ReservationRequestCanceled || ReservationResponse == ESpectatorReservationResult::ReservationNotFound);
+	ensure(bCancelReservation);
+
+	OnCancelledComplete();
+}
+
+/// @cond DOXYGEN_WARNINGS
+
+void ASpectatorBeaconClient::ClientSendReservationUpdates_Implementation(int32 NumRemainingReservations)
+{
+	if (!bCancelReservation)
+	{
+#if !UE_BUILD_SHIPPING
+		const float Rate = SpectatorBeaconConsoleVariables::CVarDelayUpdateResponse.GetValueOnGameThread();
+#else
+		const float Rate = 0.0f;
+#endif
+		if (Rate > 0.0f)
+		{
+			UE_LOG(LogSpectatorBeacon, Verbose, TEXT("Spectator beacon reservations remaining %d, waiting %fs to notify"), NumRemainingReservations, Rate);
+
+			FTimerDelegate TimerDelegate;
+			TimerDelegate.BindLambda([this, NumRemainingReservations]()
+			{
+				ProcessReservationUpdate(NumRemainingReservations);
+			});
+
+			PendingReservationUpdateTimerHandle = DelayResponse(TimerDelegate, Rate);
+		}
+		else
+		{
+			ProcessReservationUpdate(NumRemainingReservations);
+		}
+	}
+}
+
+/// @endcond
+
+void ASpectatorBeaconClient::ProcessReservationUpdate(int32 NumRemainingReservations)
+{
+	UE_LOG(LogSpectatorBeacon, Verbose, TEXT("Spectator beacon reservations remaining %d"), NumRemainingReservations);
+	ReservationCountUpdate.ExecuteIfBound(NumRemainingReservations);
+}
+
+/// @cond DOXYGEN_WARNINGS
+
+void ASpectatorBeaconClient::ClientSendReservationFull_Implementation()
+{
+	if (!bCancelReservation)
+	{
+#if !UE_BUILD_SHIPPING
+		const float Rate = SpectatorBeaconConsoleVariables::CVarDelayFullResponse.GetValueOnGameThread();
+#else
+		const float Rate = 0.0f;
+#endif
+		if (Rate > 0.0f)
+		{
+			UE_LOG(LogSpectatorBeacon, Verbose, TEXT("Spectator beacon reservations full, waiting %fs to notify"), Rate);
+
+			FTimerDelegate TimerDelegate;
+			TimerDelegate.BindLambda([this]()
+			{
+				ProcessReservationFull();
+			});
+
+			PendingReservationFullTimerHandle = DelayResponse(TimerDelegate, Rate);
+		}
+		else
+		{
+			ProcessReservationFull();
+		}
+	}
+}
+
+/// @endcond
+
+void ASpectatorBeaconClient::ProcessReservationFull()
+{
+	UE_LOG(LogSpectatorBeacon, Verbose, TEXT("Spectator beacon reservations full"));
+	ReservationFull.ExecuteIfBound();
+}
+
+FTimerHandle ASpectatorBeaconClient::DelayResponse(FTimerDelegate& Delegate, float Delay)
+{
+	FTimerHandle TimerHandle;
+
+	UWorld* World = GetWorld();
+	if (ensure(World != nullptr))
+	{
+		World->GetTimerManager().SetTimer(TimerHandle, Delegate, Delay, false);
+	}
+
+	return TimerHandle;
+}
diff --git a/Engine/Plugins/Online/OnlineSubsystemUtils/Source/OnlineSubsystemUtils/Private/SpectatorBeaconHost.cpp b/Engine/Plugins/Online/OnlineSubsystemUtils/Source/OnlineSubsystemUtils/Private/SpectatorBeaconHost.cpp
new file mode 100644
--- /dev/null
+++ b/Engine/Plugins/Online/OnlineSubsystemUtils/Source/OnlineSubsystemUtils/Private/SpectatorBeaconHost.cpp
@@ -0,0 +1,572 @@
+// Copyright 1998-2018 Epic Games, Inc. All Rights Reserved.
+
+#include "SpectatorBeaconHost.h"
+#include "Misc/CommandLine.h"
+#include "UObject/Package.h"
+#include "Engine/NetConnection.h"
+#include "OnlineSubsystemTypes.h"
+#include "OnlineSubsystemUtils.h"
+#include "SpectatorBeaconClient.h"
+
+ASpectatorBeaconHost::ASpectatorBeaconHost(const FObjectInitializer& ObjectInitializer) :
+	Super(ObjectInitializer),
+	State(NULL),
+	bLogoutOnSessionTimeout(true)
+{
+	ClientBeaconActorClass = ASpectatorBeaconClient::StaticClass();
+	BeaconTypeName = ClientBeaconActorClass->GetName();
+
+	PrimaryActorTick.bCanEverTick = true;
+	PrimaryActorTick.bAllowTickOnDedicatedServer = true;
+	PrimaryActorTick.bStartWithTickEnabled = true;
+}
+
+void ASpectatorBeaconHost::PostInitProperties()
+{
+	Super::PostInitProperties();
+#if !UE_BUILD_SHIPPING
+	// This value is set on the CDO as well on purpose
+	bLogoutOnSessionTimeout = FParse::Param(FCommandLine::Get(), TEXT("NoTimeouts")) ? false : true;
+#endif
+}
+
+bool ASpectatorBeaconHost::InitHostBeacon(int32 InMaxReservations, FName InSessionName)
+{
+	UE_LOG(LogSpectatorBeacon, Verbose, TEXT("InitHostBeacon MaxSize:%d"), InMaxReservations);
+	if (InMaxReservations > 0)
+	{
+		State = NewObject<USpectatorBeaconState>(GetTransientPackage(), GetSpectatorBeaconHostClass());
+		if (State->InitState(InMaxReservations, InSessionName))
+		{
+			return true;
+		}
+	}
+
+	return false;
+}
+
+bool ASpectatorBeaconHost::InitFromBeaconState(USpectatorBeaconState* PrevState)
+{
+	if (!State && PrevState)
+	{
+		UE_LOG(LogSpectatorBeacon, Verbose, TEXT("InitFromBeaconState MaxSize:%d"), PrevState->MaxReservations);
+		State = PrevState;
+		return true;
+	}
+
+	return false;
+}
+
+void ASpectatorBeaconHost::Tick(float DeltaTime)
+{
+	if (State)
+	{
+		UWorld* World = GetWorld();
+		IOnlineSessionPtr SessionInt = Online::GetSessionInterface(World);
+
+		if (SessionInt.IsValid())
+		{
+			FName SessionName = State->GetSessionName();
+			TArray<FSpectatorReservation>& Reservations = State->GetReservations();
+
+			FNamedOnlineSession* Session = SessionInt->GetNamedSession(SessionName);
+			if (Session)
+			{
+				TArray< TSharedPtr<const FUniqueNetId> > PlayersToLogout;
+				for (int32 ResIdx = 0; ResIdx < Reservations.Num(); ResIdx++)
+				{
+					FSpectatorReservation& SpectatorRes = Reservations[ResIdx];
+
+					// Determine if we have a client connection for the reservation 
+					bool bIsConnectedReservation = false;
+					for (int32 ClientIdx = 0; ClientIdx < ClientActors.Num(); ClientIdx++)
+					{
+						ASpectatorBeaconClient* Client = Cast<ASpectatorBeaconClient>(ClientActors[ClientIdx]);
+						if (Client)
+						{
+							const FSpectatorReservation& ClientRes = Client->GetPendingReservation();
+							if (ClientRes.SpectatorId == SpectatorRes.SpectatorId)
+							{
+								bIsConnectedReservation = true;
+								break;
+							}
+						}
+						else
+						{
+							UE_LOG(LogSpectatorBeacon, Error, TEXT("Missing SpectatorBeaconClient found in ClientActors array"));
+							ClientActors.RemoveAtSwap(ClientIdx);
+							ClientIdx--;
+						}
+					}
+
+					if (bIsConnectedReservation)
+					{
+						FPlayerReservation& PlayerEntry = SpectatorRes.Spectator;
+						PlayerEntry.ElapsedTime = 0.0f;
+					}
+					// Once a client beacon disconnects, update the elapsed time since they were found as a registrant in the game session
+					else
+					{
+						FPlayerReservation& PlayerEntry = SpectatorRes.Spectator;
+
+						// Determine if the player is the owner of the session	
+						const bool bIsSessionOwner = Session->OwningUserId.IsValid() && (*Session->OwningUserId == *PlayerEntry.UniqueId);
+
+						// Determine if the player member is registered in the game session
+						if (SessionInt->IsPlayerInSession(SessionName, *PlayerEntry.UniqueId) ||
+							// Never timeout the session owner
+							bIsSessionOwner)
+						{
+							FUniqueNetIdMatcher PlayerMatch(*PlayerEntry.UniqueId);
+							int32 FoundIdx = State->PlayersPendingJoin.IndexOfByPredicate(PlayerMatch);
+							if (FoundIdx != INDEX_NONE)
+							{
+								UE_LOG(LogSpectatorBeacon, Display, TEXT("Beacon (%s): pending player %s found in session (%s), removing."),
+									*GetName(),
+									*(PlayerEntry.UniqueId.ToDebugString()),
+									*SessionName.ToString());
+
+								// reset elapsed time since found
+								PlayerEntry.ElapsedTime = 0.0f;
+								// also remove from pending join list
+								State->PlayersPendingJoin.RemoveAtSwap(FoundIdx);
+							}
+						}
+						else
+						{
+							// update elapsed time
+							PlayerEntry.ElapsedTime += DeltaTime;
+
+							if (bLogoutOnSessionTimeout)
+							{
+								// if the player is pending it's initial join then check against TravelSessionTimeoutSecs instead
+								FUniqueNetIdMatcher PlayerMatch(*PlayerEntry.UniqueId);
+								const int32 FoundIdx = State->PlayersPendingJoin.IndexOfByPredicate(PlayerMatch);
+								const bool bIsPlayerPendingJoin = FoundIdx != INDEX_NONE;
+								// if the timeout has been exceeded then add to list of players 
+								// that need to be logged out from the beacon
+								if ((bIsPlayerPendingJoin && PlayerEntry.ElapsedTime > TravelSessionTimeoutSecs) ||
+									(!bIsPlayerPendingJoin && PlayerEntry.ElapsedTime > SessionTimeoutSecs))
+								{
+									PlayersToLogout.AddUnique(PlayerEntry.UniqueId.GetUniqueNetId());
+								}
+							}
+						}
+					}
+				}
+
+				if (bLogoutOnSessionTimeout)
+				{
+					// Logout any players that timed out
+					for (int32 LogoutIdx = 0; LogoutIdx < PlayersToLogout.Num(); LogoutIdx++)
+					{
+						bool bFound = false;
+						const TSharedPtr<const FUniqueNetId>& UniqueId = PlayersToLogout[LogoutIdx];
+						float ElapsedSessionTime = 0.f;
+						for (int32 ResIdx = 0; ResIdx < Reservations.Num(); ResIdx++)
+						{
+							const FSpectatorReservation& SpectatorRes = Reservations[ResIdx];
+							const FPlayerReservation& PlayerEntry = SpectatorRes.Spectator;
+							if (*PlayerEntry.UniqueId == *UniqueId)
+							{
+								ElapsedSessionTime = PlayerEntry.ElapsedTime;
+								bFound = true;
+								break;
+							}
+
+							if (bFound)
+							{
+								break;
+							}
+						}
+
+						UE_LOG(LogSpectatorBeacon, Display, TEXT("Beacon (%s): pending player logout due to timeout for %s, elapsed time = %0.3f, removing"),
+							*GetName(),
+							*UniqueId->ToDebugString(),
+							ElapsedSessionTime);
+						// Also remove from pending join list
+						State->PlayersPendingJoin.RemoveSingleSwap(UniqueId);
+						// let the beacon handle the logout and notifications/delegates
+						FUniqueNetIdRepl RemovedId(UniqueId);
+						HandlePlayerLogout(RemovedId);
+					}
+				}
+			}
+		}
+	}
+}
+
+void ASpectatorBeaconHost::SendReservationUpdates()
+{
+	if (State && ClientActors.Num() > 0)
+	{
+		int32 NumRemaining = State->GetRemainingReservations();
+		int32 MaxReservations = State->GetMaxReservations();
+		if (NumRemaining < MaxReservations)
+		{
+			if (NumRemaining > 0)
+			{
+				UE_LOG(LogSpectatorBeacon, Verbose, TEXT("Sending reservation update %d"), NumRemaining);
+				for (AOnlineBeaconClient* ClientActor : ClientActors)
+				{
+					ASpectatorBeaconClient* SpectatorBeaconClient = Cast<ASpectatorBeaconClient>(ClientActor);
+					if (SpectatorBeaconClient)
+					{
+						SpectatorBeaconClient->ClientSendReservationUpdates(NumRemaining);
+					}
+				}
+			}
+			else
+			{
+				UE_LOG(LogSpectatorBeacon, Verbose, TEXT("Sending reservation full"));
+				for (AOnlineBeaconClient* ClientActor : ClientActors)
+				{
+					ASpectatorBeaconClient* SpectatorBeaconClient = Cast<ASpectatorBeaconClient>(ClientActor);
+					if (SpectatorBeaconClient)
+					{
+						SpectatorBeaconClient->ClientSendReservationFull();
+					}
+				}
+			}
+		}
+	}
+}
+
+void ASpectatorBeaconHost::NewPlayerAdded(const FPlayerReservation& NewPlayer)
+{
+	if (NewPlayer.UniqueId.IsValid())
+	{
+		if (State)
+		{
+			FUniqueNetIdMatcher PlayerMatch(*NewPlayer.UniqueId);
+			int32 FoundIdx = State->PlayersPendingJoin.IndexOfByPredicate(PlayerMatch);
+			if (FoundIdx == INDEX_NONE)
+			{
+				UE_LOG(LogSpectatorBeacon, Verbose, TEXT("Beacon adding pending player %s"), *NewPlayer.UniqueId.ToDebugString());
+				State->PlayersPendingJoin.Add(NewPlayer.UniqueId.GetUniqueNetId());
+				OnNewPlayerAdded().ExecuteIfBound(NewPlayer);
+			}
+		}
+		else
+		{
+			UE_LOG(LogSpectatorBeacon, Warning, TEXT("Beacon skipping PlayersPendingJoin for beacon with no state!"));
+		}
+	}
+	else
+	{
+		UE_LOG(LogSpectatorBeacon, Warning, TEXT("Beacon skipping PlayersPendingJoin for invalid player!"));
+	}
+}
+
+void ASpectatorBeaconHost::NotifyReservationEventNextFrame(FOnReservationUpdate& ReservationEvent)
+{
+	UWorld* World = GetWorld();
+	check(World);
+
+	// Calling this on next tick to protect against re-entrance
+	World->GetTimerManager().SetTimerForNextTick(FTimerDelegate::CreateLambda([ReservationEvent]() { ReservationEvent.ExecuteIfBound(); }));
+}
+
+void ASpectatorBeaconHost::HandlePlayerLogout(const FUniqueNetIdRepl& PlayerId)
+{
+	if (PlayerId.IsValid())
+	{
+		UE_LOG(LogSpectatorBeacon, Verbose, TEXT("HandlePlayerLogout %s"), *PlayerId->ToDebugString());
+
+		if (State)
+		{
+			if (State->RemovePlayer(PlayerId))
+			{
+				SendReservationUpdates();
+				NotifyReservationEventNextFrame(ReservationChanged);
+			}
+		}
+	}
+}
+
+bool ASpectatorBeaconHost::PlayerHasReservation(const FUniqueNetId& PlayerId) const
+{
+	bool bHasReservation = false;
+	if (State)
+	{
+		bHasReservation = State->PlayerHasReservation(PlayerId);
+	}
+	else
+	{
+		UE_LOG(LogSpectatorBeacon, Warning,
+			TEXT("Beacon (%s) hasn't been initialized yet, no reservations."),
+			*GetBeaconType());
+	}
+
+	return bHasReservation;
+}
+
+bool ASpectatorBeaconHost::GetPlayerValidation(const FUniqueNetId& PlayerId, FString& OutValidation) const
+{
+	OutValidation.Empty();
+
+	bool bHasValidation = false;
+	if (State)
+	{
+		bHasValidation = State->GetPlayerValidation(PlayerId, OutValidation);
+	}
+	else
+	{
+		UE_LOG(LogSpectatorBeacon, Warning,
+			TEXT("Beacon (%s) hasn't been initialized yet, no validation."),
+			*GetBeaconType());
+	}
+
+	return bHasValidation;
+}
+
+ESpectatorReservationResult::Type ASpectatorBeaconHost::AddSpectatorReservation(const FSpectatorReservation& ReservationRequest)
+{
+	ESpectatorReservationResult::Type Result = ESpectatorReservationResult::GeneralError;
+
+	if (!State || GetBeaconState() == EBeaconState::DenyRequests)
+	{
+		return ESpectatorReservationResult::ReservationDenied;
+	}
+
+	if (ReservationRequest.IsValid())
+	{
+		TArray<FSpectatorReservation>& Reservations = State->GetReservations();
+		const int32 ExistingReservationIdx = State->GetExistingReservation(ReservationRequest.SpectatorId);
+		if (ExistingReservationIdx != INDEX_NONE)
+		{
+			FSpectatorReservation& ExistingReservation = Reservations[ExistingReservationIdx];
+			const FPlayerReservation& NewPlayerRes = ReservationRequest.Spectator;
+			if (!NewPlayerRes.ValidationStr.IsEmpty())
+			{
+				FPlayerReservation& PlayerRes = ExistingReservation.Spectator;
+				// Update the validation auth strings because they may have changed with a new login 
+				PlayerRes.ValidationStr = NewPlayerRes.ValidationStr;
+
+			}
+
+			if (State->CrossPlayAllowed(ReservationRequest))
+			{
+				SendReservationUpdates();
+
+				// Clean up the game entities for these duplicate players
+				DuplicateReservation.ExecuteIfBound(ReservationRequest);
+
+				NewPlayerAdded(ReservationRequest.Spectator);
+				Result = ESpectatorReservationResult::ReservationDuplicate;
+			}
+			else
+			{
+				// Cross play restriction
+				Result = ESpectatorReservationResult::ReservationDenied_CrossPlayRestriction;
+			}
+		}
+		else
+		{
+			// Check for players we already have reservations for
+			// Keep track of team index for existing members - if we have members on opposing teams, reject this reservation
+			bool bContainsExistingMembers = false;
+			const FPlayerReservation& Spectator = ReservationRequest.Spectator;
+			{
+				int32 MemberExistingSpectatorReservationIdx = State->GetExistingReservation(Spectator.UniqueId);
+				if (MemberExistingSpectatorReservationIdx != INDEX_NONE)
+				{
+					const FSpectatorReservation& MemberExistingSpectatorReservation = Reservations[MemberExistingSpectatorReservationIdx];
+					UE_LOG(LogSpectatorBeacon, Display, TEXT("ASpectatorBeaconHost::AddSpectatorReservation: Found existing reservation for spectator %s"), *Spectator.UniqueId.ToDebugString());
+					ReservationRequest.Dump();
+					MemberExistingSpectatorReservation.Dump();
+					bContainsExistingMembers = true;
+				}
+				else
+				{
+					// Is this player in the pending join list?
+					FUniqueNetIdMatcher PlayerMatch(*Spectator.UniqueId);
+					int32 FoundIdx = State->PlayersPendingJoin.IndexOfByPredicate(PlayerMatch);
+					if (FoundIdx != INDEX_NONE)
+					{
+						UE_LOG(LogSpectatorBeacon, Display, TEXT("ASpectatorBeaconHost::AddSpectatorReservation: Found spectator %s in the pending player list"), *Spectator.UniqueId.ToDebugString());
+						ReservationRequest.Dump();
+						bContainsExistingMembers = true;
+					}
+				}
+			}
+
+			if (!bContainsExistingMembers)
+			{
+				if (State->DoesReservationFit(ReservationRequest))
+				{
+					bool bContinue = true;
+					if (ValidatePlayers.IsBound())
+					{
+						bContinue = ValidatePlayers.Execute(ReservationRequest.Spectator);
+					}
+
+					if (bContinue)
+					{						
+						if (State->CrossPlayAllowed(ReservationRequest))
+						{
+							if (State->AddReservation(ReservationRequest))
+							{
+								// Keep track of newly added players
+								NewPlayerAdded(Spectator);
+
+								SendReservationUpdates();
+
+								NotifyReservationEventNextFrame(ReservationChanged);
+								if (State->IsBeaconFull())
+								{
+									NotifyReservationEventNextFrame(ReservationsFull);
+								}
+								Result = ESpectatorReservationResult::ReservationAccepted;
+							}
+							else
+							{
+								// Something wrong with team assignment
+								Result = ESpectatorReservationResult::IncorrectPlayerCount;
+							}
+						}
+						else
+						{
+							// Cross play restriction
+							Result = ESpectatorReservationResult::ReservationDenied_CrossPlayRestriction;
+						}
+					
+					}
+					else
+					{
+						// Validate players failed above
+						Result = ESpectatorReservationResult::ReservationDenied_Banned;
+					}
+				}
+				else
+				{
+					// Reservation doesn't fit (not enough space in general)
+					Result = ESpectatorReservationResult::SpectatorLimitReached;
+				}
+			}
+			else
+			{
+				// Reservation contains players already accounted for
+				Result = ESpectatorReservationResult::ReservationDenied_ContainsExistingPlayers;
+			}
+		}
+	}
+	else
+	{
+		// Invalid reservation
+		Result = ESpectatorReservationResult::ReservationInvalid;
+	}
+
+	return Result;
+}
+
+ESpectatorReservationResult::Type ASpectatorBeaconHost::RemoveSpectatorReservation(const FUniqueNetIdRepl& Spectator)
+{
+	if (State && State->RemoveReservation(Spectator))
+	{
+		CancelationReceived.ExecuteIfBound(*Spectator);
+
+		SendReservationUpdates();
+		NotifyReservationEventNextFrame(ReservationChanged);
+		return ESpectatorReservationResult::ReservationRequestCanceled;
+	}
+	else
+	{
+		UE_LOG(LogSpectatorBeacon, Warning, TEXT("Failed to find reservation to cancel for spectator %s:"), Spectator.IsValid() ? *Spectator->ToString() : TEXT("INVALID"));
+		return ESpectatorReservationResult::ReservationNotFound;
+	}
+}
+
+void ASpectatorBeaconHost::RegisterAuthTicket(const FUniqueNetIdRepl& InSpectatorId, const FString& InAuthTicket)
+{
+	if (State)
+	{
+		State->RegisterAuthTicket(InSpectatorId, InAuthTicket);
+	}
+	else
+	{
+		UE_LOG(LogSpectatorBeacon, Warning,
+			TEXT("Beacon (%s) hasn't been initialized yet, not able to register auth ticket."),
+			*GetBeaconType());
+	}
+}
+
+bool ASpectatorBeaconHost::DoesSessionMatch(const FString& SessionId) const
+{
+	if (State)
+	{
+		UWorld* World = GetWorld();
+		FName SessionName = State->GetSessionName();
+
+		IOnlineSessionPtr SessionInt = Online::GetSessionInterface(World);
+		FNamedOnlineSession* Session = SessionInt.IsValid() ? SessionInt->GetNamedSession(SessionName) : NULL;
+		if (Session && Session->SessionInfo.IsValid() && !SessionId.IsEmpty() && Session->SessionInfo->GetSessionId().ToString() == SessionId)
+		{
+			return true;
+		}
+	}
+
+	return false;
+}
+
+void ASpectatorBeaconHost::ProcessReservationRequest(ASpectatorBeaconClient* Client, const FString& SessionId, const FSpectatorReservation& ReservationRequest)
+{
+	UE_LOG(LogSpectatorBeacon, Verbose, TEXT("ProcessReservationRequest %s SessionId %s Spectator: %s from (%s)"),
+		Client ? *Client->GetName() : TEXT("NULL"),
+		*SessionId,
+		ReservationRequest.SpectatorId.IsValid() ? *ReservationRequest.SpectatorId->ToString() : TEXT("INVALID"),
+		Client ? *Client->GetNetConnection()->LowLevelDescribe() : TEXT("NULL"));
+	if (UE_LOG_ACTIVE(LogSpectatorBeacon, Verbose))
+	{
+		ReservationRequest.Dump();
+	}
+
+	if (Client)
+	{
+		ESpectatorReservationResult::Type Result = ESpectatorReservationResult::BadSessionId;
+		if (DoesSessionMatch(SessionId))
+		{
+			Result = AddSpectatorReservation(ReservationRequest);
+		}
+
+		UE_LOG(LogSpectatorBeacon, Verbose, TEXT("ProcessReservationRequest result: %s"), ESpectatorReservationResult::ToString(Result));
+		if (UE_LOG_ACTIVE(LogSpectatorBeacon, Verbose) &&
+			(Result != ESpectatorReservationResult::ReservationAccepted))
+		{
+			DumpReservations();
+		}
+
+		Client->ClientReservationResponse(Result);
+	}
+}
+
+void ASpectatorBeaconHost::ProcessCancelReservationRequest(ASpectatorBeaconClient* Client, const FUniqueNetIdRepl& Spectator)
+{
+	UE_LOG(LogSpectatorBeacon, Verbose, TEXT("ProcessCancelReservationRequest %s Spectator: %s from (%s)"),
+		Client ? *Client->GetName() : TEXT("NULL"),
+		Spectator.IsValid() ? *Spectator->ToString() : TEXT("INVALID"),
+		Client ? *Client->GetNetConnection()->LowLevelDescribe() : TEXT("NULL"));
+
+	if (Client)
+	{
+		ESpectatorReservationResult::Type Result = RemoveSpectatorReservation(Spectator);
+		UE_LOG(LogSpectatorBeacon, Verbose, TEXT("ProcessCancelReservationRequest result: %s"), ESpectatorReservationResult::ToString(Result));
+		if (UE_LOG_ACTIVE(LogSpectatorBeacon, Verbose) &&
+			(Result != ESpectatorReservationResult::ReservationRequestCanceled))
+		{
+			DumpReservations();
+		}
+
+		Client->ClientCancelReservationResponse(Result);
+	}
+}
+
+void ASpectatorBeaconHost::DumpReservations() const
+{
+	UE_LOG(LogSpectatorBeacon, Display, TEXT("Debug info for Beacon: %s"), *GetBeaconType());
+	if (State)
+	{
+		State->DumpReservations();
+	}
+	UE_LOG(LogSpectatorBeacon, Display, TEXT(""));
+}
diff --git a/Engine/Plugins/Online/OnlineSubsystemUtils/Source/OnlineSubsystemUtils/Private/SpectatorBeaconState.cpp b/Engine/Plugins/Online/OnlineSubsystemUtils/Source/OnlineSubsystemUtils/Private/SpectatorBeaconState.cpp
new file mode 100644
--- /dev/null
+++ b/Engine/Plugins/Online/OnlineSubsystemUtils/Source/OnlineSubsystemUtils/Private/SpectatorBeaconState.cpp
@@ -0,0 +1,341 @@
+// Copyright 1998-2018 Epic Games, Inc. All Rights Reserved.
+
+#include "SpectatorBeaconState.h"
+#include "OnlineBeacon.h"
+#include "OnlineSubsystemTypes.h"
+
+DEFINE_LOG_CATEGORY(LogSpectatorBeacon);
+
+bool FSpectatorReservation::IsValid() const
+{
+	bool bIsValid = false;
+	if (SpectatorId.IsValid())
+	{
+		bIsValid = true;
+		if (!Spectator.UniqueId.IsValid())
+		{
+			bIsValid = false;
+		}
+		else if (Spectator.Platform.IsEmpty())
+		{
+			bIsValid = false;
+		}
+		else if (SpectatorId == Spectator.UniqueId &&
+				Spectator.ValidationStr.IsEmpty())
+		{
+				bIsValid = false;
+		
+		}
+	}
+
+	return bIsValid;
+}
+
+void FSpectatorReservation::Dump() const
+{
+	UE_LOG(LogSpectatorBeacon, Display, TEXT("Spectator Reservation:"));
+	UE_LOG(LogSpectatorBeacon, Display, TEXT("  Spectator: %s"), *SpectatorId.ToDebugString());
+	UE_LOG(LogSpectatorBeacon, Display, TEXT("  UniqueId: %s"), *Spectator.UniqueId.ToDebugString());
+	UE_LOG(LogSpectatorBeacon, Display, TEXT("	Crossplay: %s"), *LexToString(Spectator.bAllowCrossplay));
+#if UE_BUILD_SHIPPING
+		UE_LOG(LogSpectatorBeacon, Display, TEXT("  ValidationStr: %d bytes"), Spectator.ValidationStr.Len());
+#else
+		UE_LOG(LogSpectatorBeacon, Display, TEXT("  ValidationStr: %s"), *Spectator.ValidationStr);
+#endif
+		UE_LOG(LogSpectatorBeacon, Display, TEXT("  ElapsedTime: %0.2f"), Spectator.ElapsedTime);
+}
+
+USpectatorBeaconState::USpectatorBeaconState(const FObjectInitializer& ObjectInitializer) :
+	Super(ObjectInitializer),
+	SessionName(NAME_None),
+	NumConsumedReservations(0),
+	MaxReservations(0),
+	bRestrictCrossConsole(false)
+{
+}
+
+bool USpectatorBeaconState::InitState(int32 InMaxReservations, FName InSessionName)
+{
+	if (InMaxReservations > 0)
+	{
+		SessionName = InSessionName;
+		MaxReservations = InMaxReservations;
+		Reservations.Empty(MaxReservations);
+		return true;
+	}
+
+	return false;
+}
+
+bool USpectatorBeaconState::HasCrossplayOptOutReservation() const
+{
+	for (const FSpectatorReservation& ExistingReservation : Reservations)
+	{
+		const FPlayerReservation& ExistingPlayer = ExistingReservation.Spectator;
+		if (!ExistingPlayer.bAllowCrossplay)
+		{
+			return true;
+		}
+	}
+
+	return false;
+}
+
+int32 USpectatorBeaconState::GetReservationPlatformCount(const FString& InPlatform) const
+{
+	int32 PlayerCount = 0;
+	for (const FSpectatorReservation& ExistingReservation : Reservations)
+	{
+		const FPlayerReservation& ExistingPlayer = ExistingReservation.Spectator;
+		if (ExistingPlayer.Platform == InPlatform)
+		{
+			PlayerCount++;
+		}
+	}
+
+	return PlayerCount;
+}
+
+bool USpectatorBeaconState::CrossPlayAllowed(const FSpectatorReservation& ReservationRequest) const
+{
+	// Since this player is a spectator, he won't be playing, so allow crossplay.
+	return true;
+}
+
+bool USpectatorBeaconState::DoesReservationFit(const FSpectatorReservation& ReservationRequest) const
+{
+	const bool bRoomForReservation = (NumConsumedReservations + 1) <= MaxReservations;
+
+	UE_LOG(LogSpectatorBeacon, Verbose, TEXT("USpectatorBeaconState::DoesReservationFit: NumConsumedReservations: %d MaxReservations: %d"), NumConsumedReservations, MaxReservations);
+
+	return bRoomForReservation;
+}
+
+bool USpectatorBeaconState::AddReservation(const FSpectatorReservation& ReservationRequest)
+{
+	if (UE_LOG_ACTIVE(LogSpectatorBeacon, Verbose))
+	{
+		UE_LOG(LogSpectatorBeacon, Verbose, TEXT("USpectatorBeaconState::AddReservation"));
+		ReservationRequest.Dump();
+	}
+	
+	NumConsumedReservations += 1;
+	UE_LOG(LogSpectatorBeacon, Verbose, TEXT("UsPECTATORBeaconState::AddReservation: Setting NumConsumedReservations to %d"), NumConsumedReservations);
+	int32 ResIdx = Reservations.Add(ReservationRequest);
+	SanityCheckReservations(false);
+
+	return true;
+}
+
+bool USpectatorBeaconState::RemoveReservation(const FUniqueNetIdRepl& Spectator)
+{
+	const int32 ExistingReservationIdx = GetExistingReservation(Spectator);
+	if (ExistingReservationIdx != INDEX_NONE)
+	{
+		NumConsumedReservations -= 1;
+		UE_LOG(LogSpectatorBeacon, Verbose, TEXT("USpectatorBeaconState::RemoveReservation: %s, setting NumConsumedReservations to %d"), *Spectator.ToString(), NumConsumedReservations);
+
+		const FPlayerReservation& PlayerRes = Reservations[ExistingReservationIdx].Spectator;
+		FUniqueNetIdMatcher PlayerMatch(*PlayerRes.UniqueId);
+		int32 FoundIdx = PlayersPendingJoin.IndexOfByPredicate(PlayerMatch);
+		if (FoundIdx != INDEX_NONE)
+		{
+			PlayersPendingJoin.RemoveAtSwap(FoundIdx);
+		}
+
+		Reservations.RemoveAtSwap(ExistingReservationIdx);
+
+		SanityCheckReservations(false);
+		return true;
+	}
+
+	return false;
+}
+
+void USpectatorBeaconState::RegisterAuthTicket(const FUniqueNetIdRepl& InSpectatorId, const FString& InAuthTicket)
+{
+	if (InSpectatorId.IsValid() && !InAuthTicket.IsEmpty())
+	{
+		bool bFoundReservation = false;
+
+		for (int32 ResIdx = 0; ResIdx < Reservations.Num() && !bFoundReservation; ResIdx++)
+		{
+			FSpectatorReservation& ReservationEntry = Reservations[ResIdx];
+
+			FPlayerReservation& PlayerRes = ReservationEntry.Spectator;
+
+			if (InSpectatorId == PlayerRes.UniqueId)
+			{
+				if (PlayerRes.ValidationStr.IsEmpty())
+				{
+					UE_LOG(LogSpectatorBeacon, VeryVerbose, TEXT("Setting auth ticket for spectator %s."), *InSpectatorId.ToDebugString());
+				}
+				else if (PlayerRes.ValidationStr != InAuthTicket)
+				{
+					UE_LOG(LogSpectatorBeacon, Warning, TEXT("Auth ticket changing for spectator %s."), *InSpectatorId.ToDebugString());
+				}
+
+				PlayerRes.ValidationStr = InAuthTicket;
+				bFoundReservation = true;
+				break;
+			}
+		}
+
+		if (!bFoundReservation)
+		{
+			UE_LOG(LogSpectatorBeacon, Warning, TEXT("Found no reservation for player %s, while registering auth ticket."), *InSpectatorId.ToDebugString());
+		}
+	}
+}
+
+
+bool USpectatorBeaconState::RemovePlayer(const FUniqueNetIdRepl& PlayerId)
+{
+	UE_LOG(LogSpectatorBeacon, VeryVerbose, TEXT("USpectatrBeaconState::RemovePlayer: %s"), *PlayerId.ToDebugString());
+	bool bWasRemoved = false;
+
+	for (int32 ResIdx = 0; ResIdx < Reservations.Num() && !bWasRemoved; ResIdx++)
+	{
+		FSpectatorReservation& Reservation = Reservations[ResIdx];
+		FPlayerReservation& PlayerEntry = Reservation.Spectator;
+		if (PlayerEntry.UniqueId == PlayerId)
+		{
+			bWasRemoved = true;
+
+			// free up a consumed entry
+			NumConsumedReservations--;
+			if (UE_LOG_ACTIVE(LogSpectatorBeacon, Verbose))
+			{
+				UE_LOG(LogSpectatorBeacon, Verbose, TEXT("USpectatorBeaconState::RemovePlayer: Player found in reservation with id %s, setting NumConsumedReservations to %d"), *Reservation.SpectatorId.ToString(), NumConsumedReservations);
+				Reservation.Dump();
+			}
+			SanityCheckReservations(true);
+			UE_LOG(LogSpectatorBeacon, Verbose, TEXT("USpectatorBeaconState::RemovePlayer: Empty reservation found with spectator %s, removing"), *Reservation.SpectatorId.ToString());
+			Reservations.RemoveAtSwap(ResIdx--);
+		}
+	}
+
+	SanityCheckReservations(false);
+	return bWasRemoved;
+}
+
+int32 USpectatorBeaconState::GetExistingReservation(const FUniqueNetIdRepl& Spectator) const
+{
+	int32 Result = INDEX_NONE;
+	for (int32 ResIdx = 0; ResIdx < Reservations.Num(); ResIdx++)
+	{
+		const FSpectatorReservation& ReservationEntry = Reservations[ResIdx];
+		if (ReservationEntry.SpectatorId == Spectator)
+		{
+			Result = ResIdx;
+			break;
+		}
+	}
+
+	return Result;
+}
+
+bool USpectatorBeaconState::UpdateMemberPlatform(const FUniqueNetIdRepl& Spectator, const FString& PlatformName)
+{
+	for (int32 ResIdx = 0; ResIdx < Reservations.Num(); ResIdx++)
+	{
+		FSpectatorReservation& ReservationEntry = Reservations[ResIdx];
+		FPlayerReservation& PlayerReservation = ReservationEntry.Spectator;
+		if (PlayerReservation.UniqueId == Spectator)
+		{
+			if (!PlatformName.IsEmpty())
+			{
+				PlayerReservation.Platform = PlatformName;
+			}
+			// Return that member was updated
+			return true;
+		}
+	}
+
+	//Return that member was not updated
+	return false;
+}
+
+bool USpectatorBeaconState::PlayerHasReservation(const FUniqueNetId& PlayerId) const
+{
+	bool bFound = false;
+
+	for (int32 ResIdx = 0; ResIdx < Reservations.Num(); ResIdx++)
+	{
+		const FSpectatorReservation& ReservationEntry = Reservations[ResIdx];
+		if (*ReservationEntry.Spectator.UniqueId == PlayerId)
+		{
+			bFound = true;
+			break;
+		}
+	}
+
+	return bFound;
+}
+
+bool USpectatorBeaconState::GetPlayerValidation(const FUniqueNetId& PlayerId, FString& OutValidation) const
+{
+	bool bFound = false;
+	OutValidation = FString();
+
+	for (int32 ResIdx = 0; ResIdx < Reservations.Num() && !bFound; ResIdx++)
+	{
+		const FSpectatorReservation& ReservationEntry = Reservations[ResIdx];
+		if (*ReservationEntry.Spectator.UniqueId == PlayerId)
+		{
+			OutValidation = ReservationEntry.Spectator.ValidationStr;
+			bFound = true;
+			break;
+		}
+	}
+
+	return bFound;
+}
+
+void USpectatorBeaconState::DumpReservations() const
+{
+	FUniqueNetIdRepl NetId;
+	FPlayerReservation PlayerRes;
+
+	UE_LOG(LogSpectatorBeacon, Display, TEXT("Session that reservations are for: %s"), *SessionName.ToString());
+	UE_LOG(LogSpectatorBeacon, Display, TEXT("Number total reservations: %d"), MaxReservations);
+	UE_LOG(LogSpectatorBeacon, Display, TEXT("Number consumed reservations: %d"), NumConsumedReservations);
+	UE_LOG(LogSpectatorBeacon, Display, TEXT("Number of spectator reservations: %d"), Reservations.Num());
+
+	// Log each spectator that has a reservation
+	for (int32 ResIndex = 0; ResIndex < Reservations.Num(); ResIndex++)
+	{
+		NetId = Reservations[ResIndex].SpectatorId;
+		UE_LOG(LogSpectatorBeacon, Display, TEXT("\t Spectator: %s"), *NetId->ToDebugString());
+		PlayerRes = Reservations[ResIndex].Spectator;
+		UE_LOG(LogSpectatorBeacon, Display, TEXT("\t  Member: %s [%s] Cross: %s"), *PlayerRes.UniqueId->ToString(), *PlayerRes.Platform, *LexToString(PlayerRes.bAllowCrossplay));
+	}
+	UE_LOG(LogSpectatorBeacon, Display, TEXT(""));
+}
+
+void USpectatorBeaconState::SanityCheckReservations(const bool bIgnoreEmptyReservations) const
+{
+#if !UE_BUILD_SHIPPING
+	// Verify that each player is only in exactly one reservation
+	TMap<FUniqueNetIdRepl, FUniqueNetIdRepl> PlayersInReservation;
+	for (const FSpectatorReservation& Reservation : Reservations)
+	{
+		if (!Reservation.SpectatorId.IsValid())
+		{
+			DumpReservations();
+			checkf(false, TEXT("Reservation does not have valid spectator!"));
+		}
+		const FPlayerReservation& PlayerReservation = Reservation.Spectator;
+		if (PlayerReservation.UniqueId.IsValid())
+		{
+			const FUniqueNetIdRepl* const ExistingReservationLeader = PlayersInReservation.Find(PlayerReservation.UniqueId);
+			if (ExistingReservationLeader != nullptr)
+			{
+				DumpReservations();
+				checkf(false, TEXT("Player %s is in multiple reservations!"), *PlayerReservation.UniqueId.ToString());
+			}
+			PlayersInReservation.Add(PlayerReservation.UniqueId, Reservation.SpectatorId);
+		}
+	}
+#endif // !UE_BUILD_SHIPPING
+}
diff --git a/Engine/Plugins/Online/OnlineSubsystemUtils/Source/OnlineSubsystemUtils/Public/OnlineBeaconReservation.h b/Engine/Plugins/Online/OnlineSubsystemUtils/Source/OnlineSubsystemUtils/Public/OnlineBeaconReservation.h
new file mode 100644
--- /dev/null
+++ b/Engine/Plugins/Online/OnlineSubsystemUtils/Source/OnlineSubsystemUtils/Public/OnlineBeaconReservation.h
@@ -0,0 +1,40 @@
+// Copyright 1998-2018 Epic Games, Inc. All Rights Reserved.
+
+#pragma once
+
+#include "CoreMinimal.h"
+#include "UObject/ObjectMacros.h"
+#include "GameFramework/Actor.h"
+#include "OnlineBeaconReservation.generated.h"
+
+/** A single player reservation */
+USTRUCT()
+struct FPlayerReservation
+{
+	GENERATED_USTRUCT_BODY()
+
+		FPlayerReservation()
+		: bAllowCrossplay(false)
+		, ElapsedTime(0.0f)
+	{}
+
+	/** Unique id for this reservation */
+	UPROPERTY(Transient)
+		FUniqueNetIdRepl UniqueId;
+
+	/** Info needed to validate user credentials when joining a server */
+	UPROPERTY(Transient)
+		FString ValidationStr;
+
+	/** Platform this user is on */
+	UPROPERTY(Transient)
+		FString Platform;
+
+	/** Does this player opt in to crossplay */
+	UPROPERTY(Transient)
+		bool bAllowCrossplay;
+
+	/** Elapsed time since player made reservation and was last seen */
+	UPROPERTY(Transient)
+		float ElapsedTime;
+};
\ No newline at end of file
diff --git a/Engine/Plugins/Online/OnlineSubsystemUtils/Source/OnlineSubsystemUtils/Public/PartyBeaconState.h b/Engine/Plugins/Online/OnlineSubsystemUtils/Source/OnlineSubsystemUtils/Public/PartyBeaconState.h
--- a/Engine/Plugins/Online/OnlineSubsystemUtils/Source/OnlineSubsystemUtils/Public/PartyBeaconState.h
+++ b/Engine/Plugins/Online/OnlineSubsystemUtils/Source/OnlineSubsystemUtils/Public/PartyBeaconState.h
@@ -7,6 +7,7 @@
 #include "UObject/Object.h"
 #include "UObject/CoreOnline.h"
 #include "GameFramework/OnlineReplStructs.h"
+#include "OnlineBeaconReservation.h"
 #include "PartyBeaconState.generated.h"
 
 ONLINESUBSYSTEMUTILS_API DECLARE_LOG_CATEGORY_EXTERN(LogPartyBeacon, Log, All);
@@ -172,38 +173,6 @@ namespace ETeamAssignmentMethod
 	extern ONLINESUBSYSTEMUTILS_API const FName Manual;
 }
 
-/** A single player reservation */
-USTRUCT()
-struct FPlayerReservation
-{
-	GENERATED_USTRUCT_BODY()
-	
-	FPlayerReservation()
-		: bAllowCrossplay(false)
-		, ElapsedTime(0.0f)
-	{}
-
-	/** Unique id for this reservation */
-	UPROPERTY(Transient)
-	FUniqueNetIdRepl UniqueId;
-
-	/** Info needed to validate user credentials when joining a server */
-	UPROPERTY(Transient)
-	FString ValidationStr;
-
-	/** Platform this user is on */
-	UPROPERTY(Transient)
-	FString Platform;
-
-    /** Does this player opt in to crossplay */
-	UPROPERTY(Transient)
-	bool bAllowCrossplay;
-
-	/** Elapsed time since player made reservation and was last seen */
-	UPROPERTY(Transient)
-	float ElapsedTime;
-};
-
 /** A whole party reservation */
 USTRUCT()
 struct ONLINESUBSYSTEMUTILS_API FPartyReservation
diff --git a/Engine/Plugins/Online/OnlineSubsystemUtils/Source/OnlineSubsystemUtils/Public/SpectatorBeaconClient.h b/Engine/Plugins/Online/OnlineSubsystemUtils/Source/OnlineSubsystemUtils/Public/SpectatorBeaconClient.h
new file mode 100644
--- /dev/null
+++ b/Engine/Plugins/Online/OnlineSubsystemUtils/Source/OnlineSubsystemUtils/Public/SpectatorBeaconClient.h
@@ -0,0 +1,258 @@
+// Copyright 1998-2018 Epic Games, Inc. All Rights Reserved.
+
+#pragma once
+
+#include "CoreMinimal.h"
+#include "UObject/ObjectMacros.h"
+#include "Engine/EngineTypes.h"
+#include "SpectatorBeaconState.h"
+#include "TimerManager.h"
+#include "OnlineBeaconClient.h"
+
+#include "SpectatorBeaconClient.generated.h"
+
+class FOnlineSessionSearchResult;
+
+/**
+* Types of reservation requests that can be made by this beacon
+*/
+UENUM()
+enum class ESpectatorClientRequestType : uint8
+{
+	/** None pending */
+	NonePending,
+	/** Make a reservation with an existing session */
+	ExistingSessionReservation,
+	/** Make an update to an existing reservation */
+	ReservationUpdate,
+	/** Reservation to configure an empty server  */
+	EmptyServerReservation,
+	/** Simple reconnect (checks for existing reservation) */
+	Reconnect,
+	/** Abandon the reservation beacon (game specific handling)*/
+	Abandon
+};
+
+inline const TCHAR* ToString(ESpectatorClientRequestType RequestType)
+{
+	switch (RequestType)
+	{
+	case ESpectatorClientRequestType::NonePending:
+	{
+		return TEXT("No Request Pending");
+	}
+	case ESpectatorClientRequestType::ExistingSessionReservation:
+	{
+		return TEXT("Existing Session Reservation");
+	}
+	case ESpectatorClientRequestType::ReservationUpdate:
+	{
+		return TEXT("Reservation Update");
+	}
+	case ESpectatorClientRequestType::EmptyServerReservation:
+	{
+		return TEXT("Empty Server Reservation");
+	}
+	case ESpectatorClientRequestType::Reconnect:
+	{
+		return TEXT("Reconnect Only");
+	}
+	}
+	return TEXT("");
+}
+
+/**
+* Delegate triggered when a response from the spectator beacon host has been received
+*
+* @param ReservationResponse response from the server
+*/
+DECLARE_DELEGATE_OneParam(FOnSpectatorReservationRequestComplete, ESpectatorReservationResult::Type /** ReservationResponse */);
+
+/**
+* Delegate triggered when the host indicated a reservation count has changed
+*
+* @param NumRemaining number of slots remaining in the session
+*/
+DECLARE_DELEGATE_OneParam(FOnReservationCountUpdate, int32 /** NumRemaining */);
+
+/** Delegate triggered when the host indicated the reservation is full */
+DECLARE_DELEGATE(FOnReservationFull);
+
+/**
+* A beacon client used for making reservations with an existing game session
+*/
+UCLASS(transient, notplaceable, config = Engine)
+class ONLINESUBSYSTEMUTILS_API ASpectatorBeaconClient : public AOnlineBeaconClient
+{
+	GENERATED_UCLASS_BODY()
+
+		//~ Begin UObject Interface
+		virtual void BeginDestroy() override;
+	//~ End UObject Interface
+
+	//~ Begin AOnlineBeaconClient Interface
+	virtual void OnConnected() override;
+	virtual void OnFailure() override;
+	//~ End AOnlineBeaconClient Interface
+
+	/**
+	* Sends a request to the remote host to allow the specified members to reserve space
+	* in the host's session. Note this request is async.
+	*
+	* @param ConnectInfoStr the URL of the server that the connection will be made to
+	* @param InSessionId Id of the session expected to be found at this destination
+	* @param RequestingSpectator the spectator that will be joining
+	* @param Reservation the reservation for the spectator who wants to reserve space
+	*
+	* @return true if the request able to be sent, false if it failed to send
+	*/
+	virtual bool RequestReservation(const FString& ConnectInfoStr, const FString& InSessionId, const FUniqueNetIdRepl& RequestingSpectator, const FPlayerReservation& Reservation);
+
+	/**
+	* Sends a request to the remote host to allow the specified members to reserve space
+	* in the host's session. Note this request is async.
+	*
+	* @param DesiredHost a search result describing the server that the connection will be made to
+	* @param RequestingSpectator the spectator that will be joining
+	* @param Reservation the reservation for the spectator who wants to reserve space
+	*
+	* @return true if the request able to be sent, false if it failed to send
+	*/
+	virtual bool RequestReservation(const FOnlineSessionSearchResult& DesiredHost, const FUniqueNetIdRepl& RequestingSpectator, const FPlayerReservation& Reservation);
+
+	/**
+	* Cancel an existing request to the remote host to revoke allocated space on the server.
+	* Note this request is async.
+	*/
+	virtual void CancelReservation();
+
+	/**
+	* Response from the host session after making a reservation request
+	*
+	* @param ReservationResponse response from server
+	*/
+	UFUNCTION(client, reliable)
+		virtual void ClientReservationResponse(ESpectatorReservationResult::Type ReservationResponse);
+
+	/**
+	* Response from the host session after making a cancellation request
+	*
+	* @param ReservationResponse response from server
+	*/
+	UFUNCTION(client, reliable)
+		virtual void ClientCancelReservationResponse(ESpectatorReservationResult::Type ReservationResponse);
+
+	/**
+	* Response from the host session that the reservation count has changed
+	*
+	* @param NumRemainingReservations number of slots remaining until a full session
+	*/
+	UFUNCTION(client, reliable)
+		virtual void ClientSendReservationUpdates(int32 NumRemainingReservations);
+
+	/** Response from the host session that the reservation is full */
+	UFUNCTION(client, reliable)
+		virtual void ClientSendReservationFull();
+
+	/**
+	* Delegate triggered when a response from the spectator beacon host has been received
+	*
+	* @return delegate to handle response from the server
+	*/
+	FOnSpectatorReservationRequestComplete& OnReservationRequestComplete() { return ReservationRequestComplete; }
+
+	/**
+	* Delegate triggered when the host indicated a reservation count has changed
+	*
+	* @param NumRemaining number of slots remaining in the session
+	*/
+	FOnReservationCountUpdate& OnReservationCountUpdate() { return ReservationCountUpdate; }
+
+	/** Delegate triggered when the host indicated the reservation is full */
+	FOnReservationFull& OnReservationFull() { return ReservationFull; }
+
+	/**
+	* @return the pending reservation associated with this beacon client
+	*/
+	const FSpectatorReservation& GetPendingReservation() const { return PendingReservation; }
+
+protected:
+
+	/** Delegate for reservation request responses */
+	FOnSpectatorReservationRequestComplete ReservationRequestComplete;
+	/** Delegate for reservation count updates */
+	FOnReservationCountUpdate ReservationCountUpdate;
+
+	/** Delegate for reservation full */
+	FOnReservationFull ReservationFull;
+
+	/** Session Id of the destination host */
+	UPROPERTY()
+		FString DestSessionId;
+	/** Pending reservation that will be sent upon connection with the intended host */
+	UPROPERTY()
+		FSpectatorReservation PendingReservation;
+
+	/** Type of request currently being handled by this client beacon */
+	UPROPERTY()
+		ESpectatorClientRequestType RequestType;
+
+	/** Has the reservation request been delivered */
+	UPROPERTY()
+		bool bPendingReservationSent;
+	/** Has the reservation request been canceled */
+	UPROPERTY()
+		bool bCancelReservation;
+
+	/** Timer to trigger a cancel reservation request if the server doesn't respond in time */
+	FTimerHandle CancelRPCFailsafe;
+
+	/** Timers for delaying various responses (debug) */
+	FTimerHandle PendingResponseTimerHandle;
+	FTimerHandle PendingCancelResponseTimerHandle;
+	FTimerHandle PendingReservationUpdateTimerHandle;
+	FTimerHandle PendingReservationFullTimerHandle;
+
+	/** Clear out all the timer handles listed above */
+	void ClearTimers(bool bCallFailSafeIfNeeded);
+	/** Delegate triggered if the client doesn't hear from the server in time */
+	void OnCancelledFailsafe();
+	/** Delegate triggered when a cancel reservation request is complete */
+	void OnCancelledComplete();
+
+	/** Process a response to our RequestReservation request to the server */
+	void ProcessReservationResponse(ESpectatorReservationResult::Type ReservationResponse);
+	/** Process a response to our CancelReservation request to the server */
+	void ProcessCancelReservationResponse(ESpectatorReservationResult::Type ReservationResponse);
+	/** Process a response from the server with an update to the number of consumed reservations */
+	void ProcessReservationUpdate(int32 NumRemainingReservations);
+	/** Process a response from the server that the reservation beacon is full */
+	void ProcessReservationFull();
+
+	/**
+	* Tell the server about the reservation request being made
+	*
+	* @param SessionId expected session id on the other end (must match)
+	* @param Reservation pending reservation request to make with server
+	*/
+	UFUNCTION(server, reliable, WithValidation)
+		virtual void ServerReservationRequest(const FString& SessionId, const FSpectatorReservation& Reservation);
+
+	/**
+	* Tell the server to cancel a pending or existing reservation
+	*
+	* @param Spectator id of the spectator for the reservation to cancel
+	*/
+	UFUNCTION(server, reliable, WithValidation)
+		virtual void ServerCancelReservationRequest(const FUniqueNetIdRepl& Spectator);
+
+	/**
+	* Trigger the given delegate at a later time
+	*
+	* @param Delegate function to call at a later date
+	* @param Delay time to wait before calling function
+	*
+	* @return handle in the timer system for this entry
+	*/
+	FTimerHandle DelayResponse(FTimerDelegate& Delegate, float Delay);
+};
diff --git a/Engine/Plugins/Online/OnlineSubsystemUtils/Source/OnlineSubsystemUtils/Public/SpectatorBeaconHost.h b/Engine/Plugins/Online/OnlineSubsystemUtils/Source/OnlineSubsystemUtils/Public/SpectatorBeaconHost.h
new file mode 100644
--- /dev/null
+++ b/Engine/Plugins/Online/OnlineSubsystemUtils/Source/OnlineSubsystemUtils/Public/SpectatorBeaconHost.h
@@ -0,0 +1,279 @@
+// Copyright 1998-2018 Epic Games, Inc. All Rights Reserved.
+
+#pragma once
+
+#include "CoreMinimal.h"
+#include "UObject/ObjectMacros.h"
+#include "Templates/SubclassOf.h"
+#include "GameFramework/OnlineReplStructs.h"
+#include "SpectatorBeaconState.h"
+#include "OnlineBeaconHostObject.h"
+#include "SpectatorBeaconHost.generated.h"
+
+class ASpectatorBeaconClient;
+
+/**
+* Delegate type for handling reservation additions/removals, or full events
+*/
+DECLARE_DELEGATE(FOnReservationUpdate);
+
+/**
+* Delegate fired when a the beacon host has been told to cancel a reservation
+*
+* @param Spectator spectator cancelling reservation
+*/
+DECLARE_DELEGATE_OneParam(FOnCancelationReceived, const FUniqueNetId&);
+
+/**
+* Delegate fired when a the beacon host has added a new player
+*
+* @param PlayerId Player who reservation has been made for.
+*/
+DECLARE_DELEGATE_OneParam(FOnNewPlayerAdded, const FPlayerReservation&);
+
+/**
+* Delegate called when the beacon gets any request, allowing the owner to validate players at a higher level (bans,etc)
+*
+* @param Spectator player making up the reservation
+* @return true if these players are ok to join
+*/
+DECLARE_DELEGATE_RetVal_OneParam(bool, FOnSpectatorValidatePlayers, const FPlayerReservation&);
+
+/**
+* Delegate fired when a the beacon host detects a duplicate reservation
+*
+* @param SpectatorReservation reservation that is found to be duplicated
+*/
+DECLARE_DELEGATE_OneParam(FOnSpectatorDuplicateReservation, const FSpectatorReservation&);
+
+/**
+* A beacon host used for taking reservations for an existing game session
+*/
+UCLASS(transient, notplaceable, config = Engine)
+class ONLINESUBSYSTEMUTILS_API ASpectatorBeaconHost : public AOnlineBeaconHostObject
+{
+	GENERATED_UCLASS_BODY()
+
+		//~ Begin UObject interface
+		virtual void PostInitProperties() override;
+	//~ End UObject interface
+
+	//~ Begin AActor Interface
+	virtual void Tick(float DeltaTime) override;
+	//~ End AActor Interface
+
+	//~ Begin AOnlineBeaconHostObject Interface 
+	//~ End AOnlineBeaconHost Interface 
+
+	/**
+	* Initialize the spectator host beacon
+	*
+	* @param InMaxReservation max number of reservations allowed
+	* @param InSessionName name of session related to the beacon
+	*
+	* @return true if successful created, false otherwise
+	*/
+	virtual bool InitHostBeacon(int32 InMaxReservations, FName InSessionName);
+
+	/**
+	* Initialize the spectator host beacon from a previous state/configuration
+	* all existing reservations and configuration values are preserved
+	*
+	* @return true if successful created, false otherwise
+	*/
+	virtual bool InitFromBeaconState(USpectatorBeaconState* PrevState);
+
+	/**
+	* @return reference to the state of the SpectatorBeacon
+	*/
+	USpectatorBeaconState* GetState() const { return State; }
+
+	/**
+	* Notify the beacon of a player logout
+	*
+	* @param PlayerId UniqueId of player logging out
+	*/
+	virtual void HandlePlayerLogout(const FUniqueNetIdRepl& PlayerId);
+
+	/**
+	* Get the current reservation count inside the beacon
+	* NOTE: This is *NOT* the same as the number of consumed reservations across all parties, just the total number of reservations!
+	*
+	* @return number of reservations inside the beacon (*NOT* number of consumed reservations)
+	*/
+	virtual int32 GetReservationCount() const { return State->GetReservationCount(); }
+
+	/**
+	* Get the number of reservations actually used/consumed across all parties inside the beacon
+	*
+	* @return the number of actually used reservations across all parties inside the beacon
+	*/
+	virtual int32 GetNumConsumedReservations() const { return State->GetNumConsumedReservations(); }
+
+	/**
+	* Get the maximum number of reservations allowed inside the beacon
+	*
+	* @return The maximum number of reservations allowed inside the beacon
+	*/
+	virtual int32 GetMaxReservations() const { return State->GetMaxReservations(); }
+
+	/**
+	* Does a given player id have an existing reservation
+	*
+	* @param PlayerId uniqueid of the player to check
+	*
+	* @return true if a reservation exists, false otherwise
+	*/
+	virtual bool PlayerHasReservation(const FUniqueNetId& PlayerId) const;
+
+	/**
+	* Obtain player validation string from spectator reservation entry
+	*
+	* @param PlayerId unique id of player to find validation in an existing reservation
+	* @param OutValidation [out] validation string used when player requested a reservation
+	*
+	* @return true if reservation exists for player
+	*
+	*/
+	virtual bool GetPlayerValidation(const FUniqueNetId& PlayerId, FString& OutValidation) const;
+
+	/**
+	* Attempts to add a spectator reservation to the beacon
+	*
+	* @param ReservationRequest reservation attempt
+	*
+	* @return add attempt result
+	*/
+	virtual ESpectatorReservationResult::Type AddSpectatorReservation(const FSpectatorReservation& ReservationRequest);
+
+	/**
+	* Attempts to remove a spectator reservation from the beacon
+	*
+	* @param Spectator reservation owner
+	*/
+	virtual ESpectatorReservationResult::Type RemoveSpectatorReservation(const FUniqueNetIdRepl& Spectator);
+
+	/**
+	* Register user auth ticket with the reservation system
+	* Must have an existing reservation entry
+	*
+	* @param InSpectatorId id of player logging in
+	* @param InAuthTicket auth ticket reported by the user
+	*/
+	void RegisterAuthTicket(const FUniqueNetIdRepl& InSpectatorId, const FString& InAuthTicket);
+
+	/**
+	* Handle a reservation request received from an incoming client
+	*
+	* @param Client client beacon making the request
+	* @param SessionId id of the session that is being checked
+	* @param ReservationRequest payload of request
+	*/
+	virtual void ProcessReservationRequest(ASpectatorBeaconClient* Client, const FString& SessionId, const FSpectatorReservation& ReservationRequest);
+
+	/**
+	* Handle a reservation cancellation request received from an incoming client
+	*
+	* @param Client client beacon making the request
+	* @param Spectator reservation owner
+	*/
+	virtual void ProcessCancelReservationRequest(ASpectatorBeaconClient* Client, const FUniqueNetIdRepl& Spectator);
+
+	/**
+	* Crossplay
+	*/
+
+
+	/**
+	* Delegate fired when a the beacon host detects a reservation addition/removal
+	*/
+	FOnReservationUpdate& OnReservationsFull() { return ReservationsFull; }
+
+	/**
+	* Delegate fired when a the beacon host detects that all reservations are full
+	*/
+	FOnReservationUpdate& OnReservationChanged() { return ReservationChanged; }
+
+	/**
+	* Delegate fired when the beacon host adds a new player
+	*/
+	FOnNewPlayerAdded& OnNewPlayerAdded() { return NewPlayerAddedDelegate; }
+
+	/**
+	* Delegate fired when a the beacon host cancels a reservation
+	*/
+	FOnCancelationReceived& OnCancelationReceived() { return CancelationReceived; }
+
+	/**
+	* Delegate fired when a the beacon detects a duplicate reservation
+	*/
+	FOnSpectatorDuplicateReservation& OnDuplicateReservation() { return DuplicateReservation; }
+
+	/**
+	* Delegate called when the beacon gets any request, allowing the owner to validate players at a higher level (bans,etc)
+	*/
+	FOnSpectatorValidatePlayers& OnValidatePlayers() { return ValidatePlayers; }
+
+	/**
+	* Output current state of reservations to log
+	*/
+	virtual void DumpReservations() const;
+
+protected:
+
+	/** State of the beacon */
+	UPROPERTY()
+		USpectatorBeaconState* State;
+
+	/** Delegate fired when the beacon indicates all reservations are taken */
+	FOnReservationUpdate ReservationsFull;
+	/** Delegate fired when the beacon indicates a reservation add/remove */
+	FOnReservationUpdate ReservationChanged;
+	/** Delegate fired when the beacon indicates a reservation cancellation */
+	FOnCancelationReceived CancelationReceived;
+	/** Delegate fired when the beacon detects a duplicate reservation */
+	FOnSpectatorDuplicateReservation DuplicateReservation;
+	/** Delegate fired when reservation has been added */
+	FOnNewPlayerAdded NewPlayerAddedDelegate;
+	/** Delegate fired when asking the beacon owner if this reservation is legit */
+	FOnSpectatorValidatePlayers ValidatePlayers;
+
+	/** Do the timeouts below cause a player to be removed from the reservation list */
+	UPROPERTY(Config)
+		bool bLogoutOnSessionTimeout;
+	/** Seconds that can elapse before a reservation is removed due to player not being registered with the session */
+	UPROPERTY(Transient, Config)
+		float SessionTimeoutSecs;
+	/** Seconds that can elapse before a reservation is removed due to player not being registered with the session during a travel */
+	UPROPERTY(Transient, Config)
+		float TravelSessionTimeoutSecs;
+
+	/**
+	* @return the class of the state object inside this beacon
+	*/
+	virtual TSubclassOf<USpectatorBeaconState> GetSpectatorBeaconHostClass() const { return USpectatorBeaconState::StaticClass(); }
+
+	/**
+	* Update clients with current reservation information
+	*/
+	void SendReservationUpdates();
+
+	/**
+	* Handle a newly added player
+	*
+	* @param NewPlayer reservation of newly joining player
+	*/
+	void NewPlayerAdded(const FPlayerReservation& NewPlayer);
+
+	/**
+	* Does the session match the one associated with this beacon
+	*
+	* @param SessionId session to compare
+	*
+	* @return true if the session matches, false otherwise
+	*/
+	bool DoesSessionMatch(const FString& SessionId) const;
+
+	void NotifyReservationEventNextFrame(FOnReservationUpdate& ReservationEvent);
+};
+
diff --git a/Engine/Plugins/Online/OnlineSubsystemUtils/Source/OnlineSubsystemUtils/Public/SpectatorBeaconState.h b/Engine/Plugins/Online/OnlineSubsystemUtils/Source/OnlineSubsystemUtils/Public/SpectatorBeaconState.h
new file mode 100644
--- /dev/null
+++ b/Engine/Plugins/Online/OnlineSubsystemUtils/Source/OnlineSubsystemUtils/Public/SpectatorBeaconState.h
@@ -0,0 +1,389 @@
+// Copyright 1998-2018 Epic Games, Inc. All Rights Reserved.
+
+#pragma once
+
+#include "CoreMinimal.h"
+#include "UObject/ObjectMacros.h"
+#include "UObject/Object.h"
+#include "UObject/CoreOnline.h"
+#include "GameFramework/OnlineReplStructs.h"
+#include "OnlineBeaconReservation.h"
+#include "SpectatorBeaconState.generated.h"
+
+ONLINESUBSYSTEMUTILS_API DECLARE_LOG_CATEGORY_EXTERN(LogSpectatorBeacon, Log, All);
+
+/** The result code that will be returned during spectator reservation */
+UENUM()
+namespace ESpectatorReservationResult
+{
+	enum Type
+	{
+		/** Empty state. */
+		NoResult,
+		/** Pending request due to async operation, server will contact client shortly. */
+		RequestPending,
+		/** An unknown error happened. */
+		GeneralError,
+		/** All available reservations are booked. */
+		SpectatorLimitReached,
+		/** Wrong number of players to join the session. */
+		IncorrectPlayerCount,
+		/** No response from the host. */
+		RequestTimedOut,
+		/** Already have a reservation entry for the requesting spectator. */
+		ReservationDuplicate,
+		/** Couldn't find the spectator specified for a reservation update request. */
+		ReservationNotFound,
+		/** Space was available and it's time to join. */
+		ReservationAccepted,
+		/** The beacon is paused and not accepting new connections. */
+		ReservationDenied,
+		/** Some kind of cross play restriction prevents this spectator from joining */
+		ReservationDenied_CrossPlayRestriction,
+		/** This player is banned. */
+		ReservationDenied_Banned,
+		/** The reservation request was canceled before being sent. */
+		ReservationRequestCanceled,
+		// The reservation was rejected because it was badly formed
+		ReservationInvalid,
+		// The reservation was rejected because this was the wrong session
+		BadSessionId,
+		/** The reservation contains players already in this game */
+		ReservationDenied_ContainsExistingPlayers,
+	};
+}
+
+namespace ESpectatorReservationResult
+{
+	/** @return the stringified version of the enum passed in */
+	inline const TCHAR* ToString(ESpectatorReservationResult::Type SessionType)
+	{
+		switch (SessionType)
+		{
+		case NoResult:
+		{
+			return TEXT("No outstanding request");
+		}
+		case RequestPending:
+		{
+			return TEXT("Pending Request");
+		}
+		case GeneralError:
+		{
+			return TEXT("General Error");
+		}
+		case SpectatorLimitReached:
+		{
+			return TEXT("Spectator Limit Reached");
+		}
+		case IncorrectPlayerCount:
+		{
+			return TEXT("Incorrect Player Count");
+		}
+		case RequestTimedOut:
+		{
+			return TEXT("Request Timed Out");
+		}
+		case ReservationDuplicate:
+		{
+			return TEXT("Reservation Duplicate");
+		}
+		case ReservationNotFound:
+		{
+			return TEXT("Reservation Not Found");
+		}
+		case ReservationAccepted:
+		{
+			return TEXT("Reservation Accepted");
+		}
+		case ReservationDenied:
+		{
+			return TEXT("Reservation Denied");
+		}
+		case ReservationDenied_CrossPlayRestriction:
+		{
+			return TEXT("Reservation Denied CrossPlayRestriction");
+		}
+		case ReservationDenied_Banned:
+		{
+			return TEXT("Reservation Banned");
+		}
+		case ReservationRequestCanceled:
+		{
+			return TEXT("Request Canceled");
+		}
+		case ReservationInvalid:
+		{
+			return TEXT("Invalid reservation");
+		}
+		case BadSessionId:
+		{
+			return TEXT("Bad Session Id");
+		}
+		case ReservationDenied_ContainsExistingPlayers:
+		{
+			return TEXT("Reservation Contains Existing Players");
+		}
+		}
+		return TEXT("");
+	}
+
+	inline FText GetDisplayString(ESpectatorReservationResult::Type Response)
+	{
+		switch (Response)
+		{
+		case ESpectatorReservationResult::IncorrectPlayerCount:
+		case ESpectatorReservationResult::SpectatorLimitReached:
+			return NSLOCTEXT("ESpectatorReservationResult", "FullGame", "Game full");
+		case ESpectatorReservationResult::RequestTimedOut:
+			return NSLOCTEXT("ESpectatorReservationResult", "NoResponse", "No response");
+		case ESpectatorReservationResult::ReservationDenied:
+			return NSLOCTEXT("ESpectatorReservationResult", "DeniedResponse", "Not accepting connections");
+		case ESpectatorReservationResult::ReservationDenied_Banned:
+			return NSLOCTEXT("ESpectatorReservationResult", "BannedResponse", "Player Banned");
+		case ESpectatorReservationResult::GeneralError:
+			return NSLOCTEXT("ESpectatorReservationResult", "GeneralError", "Unknown Error");
+		case ESpectatorReservationResult::ReservationNotFound:
+			return NSLOCTEXT("ESpectatorReservationResult", "ReservationNotFound", "No Reservation");
+		case ESpectatorReservationResult::ReservationAccepted:
+			return NSLOCTEXT("ESpectatorReservationResult", "Accepted", "Accepted");
+		case ESpectatorReservationResult::ReservationDuplicate:
+			return NSLOCTEXT("ESpectatorReservationResult", "DuplicateReservation", "Duplicate reservation detected");
+		case ESpectatorReservationResult::ReservationInvalid:
+			return NSLOCTEXT("ESpectatorReservationResult", "InvalidReservation", "Bad reservation request");
+		case ESpectatorReservationResult::ReservationDenied_ContainsExistingPlayers:
+			return NSLOCTEXT("ESpectatorReservationResult", "ContainsExistingPlayers", "Spectator already in session");
+		case ESpectatorReservationResult::NoResult:
+		case ESpectatorReservationResult::BadSessionId:
+		default:
+			return FText::GetEmpty();
+		}
+	}
+}
+
+/** A whole Spectator reservation */
+USTRUCT()
+struct ONLINESUBSYSTEMUTILS_API FSpectatorReservation
+{
+	GENERATED_USTRUCT_BODY()
+
+		FSpectatorReservation()
+		{}
+
+
+	/** Player initiating the request */
+	UPROPERTY(Transient)
+		FUniqueNetIdRepl SpectatorId;
+
+	/** Spectator reservation */
+	UPROPERTY(Transient)
+		FPlayerReservation Spectator;
+
+	/** Is this data well formed */
+	bool IsValid() const;
+
+	/** Dump this reservation to log */
+	void Dump() const;
+};
+
+/**
+* A beacon host used for taking reservations for an existing game session
+*/
+UCLASS(transient, notplaceable, config = Engine)
+class ONLINESUBSYSTEMUTILS_API USpectatorBeaconState : public UObject
+{
+	GENERATED_UCLASS_BODY()
+
+		/**
+		* Initialize this state object
+		*
+		* @param InTeamCount number of teams to make room for
+		* @param InTeamSize size of each team
+		* @param InMaxReservation max number of reservations allowed
+		* @param InSessionName name of session related to the beacon
+		* @param InForceTeamNum team to force players on if applicable (usually only 1 team games)
+		*
+		* @return true if successful created, false otherwise
+		*/
+		virtual bool InitState(int32 InMaxReservations, FName InSessionName);
+
+	/**
+	* Add a reservation to the beacon state, tries to assign a team
+	*
+	* @param ReservationRequest reservation to possibly add to this state
+	*
+	* @return true if successful, false otherwise
+	*/
+	virtual bool AddReservation(const FSpectatorReservation& ReservationRequest);
+
+	/**
+	* Remove an entire reservation from this state object
+	*
+	* @param Spectator player holding reservation
+	*
+	* @return true if successful, false if reservation not found
+	*/
+	virtual bool RemoveReservation(const FUniqueNetIdRepl& Spectator);
+
+	/**
+	* Register user auth ticket with the reservation system
+	* Must have an existing reservation entry
+	*
+	* @param InSpectatorId id of player logging in
+	* @param InAuthTicket auth ticket reported by the user
+	*/
+	void RegisterAuthTicket(const FUniqueNetIdRepl& InSpectatorId, const FString& InAuthTicket);
+
+	/**
+	* Remove a single player from their reservation
+	*
+	* PlayerId player to remove
+	*
+	* @return true if player found and removed, false otherwise
+	*/
+	virtual bool RemovePlayer(const FUniqueNetIdRepl& PlayerId);
+
+	/**
+	* @return the name of the session associated with this beacon state
+	*/
+	virtual FName GetSessionName() const { return SessionName; }
+
+	/**
+	* @return all reservations in this beacon state
+	*/
+	virtual TArray<FSpectatorReservation>& GetReservations() { return Reservations; }
+
+	/**
+	* Get an existing reservation for a given spectator
+	*
+	* @param Spectator UniqueId of spectator for a reservation
+	*
+	* @return index of reservation, INDEX_NONE otherwise
+	*/
+	virtual int32 GetExistingReservation(const FUniqueNetIdRepl& Spectator) const;
+
+	/**
+	* Get the current reservation count inside the beacon
+	* this is NOT the number of players in the game
+	*
+	* @return number of consumed reservations
+	*/
+	virtual int32 GetReservationCount() const { return Reservations.Num(); }
+
+	/**
+	* @return the total number of reservations allowed
+	*/
+	virtual int32 GetMaxReservations() const { return MaxReservations; }
+
+	/**
+	* @return the amount of space left in the beacon
+	*/
+	virtual int32 GetRemainingReservations() const { return MaxReservations - NumConsumedReservations; }
+
+	/**
+	* @return the number of actually used reservations across all parties
+	*/
+	virtual int32 GetNumConsumedReservations() const { return NumConsumedReservations; }
+
+	/**
+	* Determine if this reservation fits all rules for fitting in the game
+	*
+	* @param ReservationRequest reservation to test for
+	*
+	* @return true if reservation can be added to this state, false otherwise
+	*/
+	virtual bool DoesReservationFit(const FSpectatorReservation& ReservationRequest) const;
+
+	/**
+	* @return true if the beacon is currently at max capacity
+	*/
+	virtual bool IsBeaconFull() const { return NumConsumedReservations == MaxReservations; }
+
+	/**
+	* Does a given player id have an existing reservation
+	*
+	* @param PlayerId uniqueid of the player to check
+	*
+	* @return true if a reservation exists, false otherwise
+	*/
+	virtual bool PlayerHasReservation(const FUniqueNetId& PlayerId) const;
+
+	/**
+	* Updates the platform on an existing reservation
+	* (Used when MMS does not set a platform when handing out a match assignment)
+	*
+	* @param PlayerId uniqueid of the player to update
+	* @param PlatformName platform to set new player as using
+	*
+	* @return true if a reservation exists, false otherwise
+	*/
+	virtual bool UpdateMemberPlatform(const FUniqueNetIdRepl& Spectator, const FString& PlatformName);
+
+	/**
+	* Obtain player validation string from spectator reservation entry
+	*
+	* @param PlayerId unique id of player to find validation in an existing reservation
+	* @param OutValidation [out] validation string used when player requested a reservation
+	*
+	* @return true if reservation exists for player
+	*
+	*/
+	virtual bool GetPlayerValidation(const FUniqueNetId& PlayerId, FString& OutValidation) const;
+
+	/**
+	* Perform a cross play compatible check for the new reservation against existing reservations
+	*
+	* @param ReservationRequest new reservation to compare
+	*
+	* @return true if there are no cross play restrictions, false otherwise
+	*/
+	virtual bool CrossPlayAllowed(const FSpectatorReservation& ReservationRequest) const;
+
+	/**
+	* @return true if there are cross play restrictions, false otherwise
+	*/
+	virtual bool HasCrossplayOptOutReservation() const;
+
+	/**
+	* Get a count of all players for a given platform
+	*
+	* @param InPlatform platform to get a count for
+	*
+	* @return number of players for a given platform
+	*/
+	virtual int32 GetReservationPlatformCount(const FString& InPlatform) const;
+
+	/**
+	* Output current state of reservations to log
+	*/
+	virtual void DumpReservations() const;
+
+protected:
+
+	/** Session tied to the beacon */
+	UPROPERTY(Transient)
+		FName SessionName;
+	/** Number of currently consumed reservations */
+	UPROPERTY(Transient)
+		int32 NumConsumedReservations;
+	/** Maximum allowed reservations */
+	UPROPERTY(Transient)
+		int32 MaxReservations;
+	/** Are multiple consoles types allowed to play together */
+	UPROPERTY(Config)
+		bool bRestrictCrossConsole;
+
+	/** Current reservations in the system */
+	UPROPERTY(Transient)
+		TArray<FSpectatorReservation> Reservations;
+	/** Players that are expected to join shortly */
+	TArray< TSharedPtr<const FUniqueNetId> > PlayersPendingJoin;
+
+	/**
+	* Check that our reservations are in a good state
+	* @param bIgnoreEmptyReservations Whether we want to ignore empty reservations or not (because we intend to clean them up later)
+	*/
+	void SanityCheckReservations(const bool bIgnoreEmptyReservations) const;
+
+	friend class ASpectatorBeaconHost;
+};
