fix name collision for spectator beacon console variables with revervation beacon console variables.

diff --git a/Engine/Plugins/Online/OnlineSubsystemUtils/Source/OnlineSubsystemUtils/Private/SpectatorBeaconClient.cpp b/Engine/Plugins/Online/OnlineSubsystemUtils/Source/OnlineSubsystemUtils/Private/SpectatorBeaconClient.cpp
--- a/Engine/Plugins/Online/OnlineSubsystemUtils/Source/OnlineSubsystemUtils/Private/SpectatorBeaconClient.cpp
+++ b/Engine/Plugins/Online/OnlineSubsystemUtils/Source/OnlineSubsystemUtils/Private/SpectatorBeaconClient.cpp
@@ -9,7 +9,7 @@ namespace SpectatorBeaconConsoleVariables
 {
 	/** Time to delay delegates firing a reservation request response */
 	TAutoConsoleVariable<float> CVarDelayReservationResponse(
-		TEXT("beacon.DelayReservationResponse"),
+		TEXT("spectatorbeacon.DelayReservationResponse"),
 		0.0f,
 		TEXT("Delay time between received response and notification\n")
 		TEXT("Time in secs"),
@@ -17,7 +17,7 @@ namespace SpectatorBeaconConsoleVariables
 
 	/** Time to delay delegates firing a cancel reservation request response */
 	TAutoConsoleVariable<float> CVarDelayCancellationResponse(
-		TEXT("beacon.DelayCancellationResponse"),
+		TEXT("spectatorbeacon.DelayCancellationResponse"),
 		0.0f,
 		TEXT("Delay time between received cancel response and notification\n")
 		TEXT("Time in secs"),
@@ -25,7 +25,7 @@ namespace SpectatorBeaconConsoleVariables
 
 	/** Time to delay delegates firing a reservation update response */
 	TAutoConsoleVariable<float> CVarDelayUpdateResponse(
-		TEXT("beacon.DelayUpdateResponse"),
+		TEXT("spectatorbeacon.DelayUpdateResponse"),
 		0.0f,
 		TEXT("Delay time between received update response and notification\n")
 		TEXT("Time in secs"),
@@ -33,7 +33,7 @@ namespace SpectatorBeaconConsoleVariables
 
 	/** Time to delay delegates firing a reservation full response */
 	TAutoConsoleVariable<float> CVarDelayFullResponse(
-		TEXT("beacon.DelayFullResponse"),
+		TEXT("spectatorbeacon.DelayFullResponse"),
 		0.0f,
 		TEXT("Delay time between received full response and notification\n")
 		TEXT("Time in secs"),
