Extra logging for the IpNetDriver receive thread
Actually get socket error strings in test & shipping builds

diff --git a/Engine/Plugins/Online/OnlineSubsystemUtils/Source/OnlineSubsystemUtils/Private/IpNetDriver.cpp b/Engine/Plugins/Online/OnlineSubsystemUtils/Source/OnlineSubsystemUtils/Private/IpNetDriver.cpp
--- a/Engine/Plugins/Online/OnlineSubsystemUtils/Source/OnlineSubsystemUtils/Private/IpNetDriver.cpp
+++ b/Engine/Plugins/Online/OnlineSubsystemUtils/Source/OnlineSubsystemUtils/Private/IpNetDriver.cpp
@@ -809,17 +809,25 @@ void UIpNetDriver::LowLevelDestroy()
 			IpServerConnection->WaitForSendTasks();
 		}
 
+		ISocketSubsystem* SocketSubsystem = GetSocketSubsystem();
+
 		// If using a recieve thread, shut down the socket, which will signal the thread to exit gracefully, then wait on the thread.
 		if (SocketReceiveThread.IsValid() && SocketReceiveThreadRunnable.IsValid())
 		{
+			UE_LOG(LogNet, Log, TEXT("Shutting down and waiting for socket receive thread for %s"), *GetDescription());
+
 			SocketReceiveThreadRunnable->bIsRunning = false;
-			Socket->Shutdown(ESocketShutdownMode::Read);
+			
+			if (!Socket->Shutdown(ESocketShutdownMode::Read))
+			{
+				const ESocketErrors ShutdownError = SocketSubsystem->GetLastErrorCode();
+				UE_LOG(LogNet, Log, TEXT("UIpNetDriver::LowLevelDestroy Socket->Shutdown returned error %s (%d) for %s"), SocketSubsystem->GetSocketError(ShutdownError), static_cast<int>(ShutdownError), *GetDescription());
+			}
 
 			SCOPE_CYCLE_COUNTER(STAT_IpNetDriver_Destroy_WaitForReceiveThread);
 			SocketReceiveThread->WaitForCompletion();
 		}
 
-		ISocketSubsystem* SocketSubsystem = GetSocketSubsystem();
 		if( !Socket->Close() )
 		{
 			UE_LOG(LogExit, Log, TEXT("closesocket error (%i)"), (int32)SocketSubsystem->GetLastErrorCode() );
@@ -930,7 +938,7 @@ uint32 UIpNetDriver::FReceiveThreadRunnable::Run()
 {
 	const FTimespan Timeout = FTimespan::FromMilliseconds(CVarNetIpNetDriverReceiveThreadPollTimeMS.GetValueOnAnyThread());
 
-	UE_LOG(LogNet, Log, TEXT("Receive Thread Startup."));
+	UE_LOG(LogNet, Log, TEXT("UIpNetDriver::FReceiveThreadRunnable::Run starting up."));
 
 	while (bIsRunning && OwningNetDriver->Socket)
 	{
@@ -985,10 +993,14 @@ uint32 UIpNetDriver::FReceiveThreadRunnable::Run()
 				IncomingPacket.Error = WaitError;
 				IncomingPacket.PlatformTimeSeconds = FPlatformTime::Seconds();
 
+				UE_LOG(LogNet, Log, TEXT("UIpNetDriver::FReceiveThreadRunnable::Run: Socket->Wait returned error %s (%d)"), SocketSubsystem->GetSocketError(IncomingPacket.Error), static_cast<int>(IncomingPacket.Error));
+
 				ReceiveQueue.Enqueue(MoveTemp(IncomingPacket));
 			}
 		}
 	}
 
+	UE_LOG(LogNet, Log, TEXT("UIpNetDriver::FReceiveThreadRunnable::Run returning."));
+
 	return 0;
 }
diff --git a/Engine/Source/Runtime/Sockets/Private/SocketSubsystem.cpp b/Engine/Source/Runtime/Sockets/Private/SocketSubsystem.cpp
--- a/Engine/Source/Runtime/Sockets/Private/SocketSubsystem.cpp
+++ b/Engine/Source/Runtime/Sockets/Private/SocketSubsystem.cpp
@@ -387,7 +387,6 @@ FResolveInfoCached* ISocketSubsystem::CreateResolveInfoCached(TSharedPtr<FIntern
  */
 const TCHAR* ISocketSubsystem::GetSocketError(ESocketErrors Code)
 {
-#if !(UE_BUILD_SHIPPING || UE_BUILD_TEST)
 	if (Code == SE_GET_LAST_ERROR_CODE)
 	{
 		Code = GetLastErrorCode();
@@ -446,11 +445,12 @@ const TCHAR* ISocketSubsystem::GetSocketError(ESocketErrors Code)
 		case SE_TRY_AGAIN: return TEXT("SE_TRY_AGAIN");
 		case SE_NO_RECOVERY: return TEXT("SE_NO_RECOVERY");
 		case SE_NO_DATA: return TEXT("SE_NO_DATA");
+		case SE_UDP_ERR_PORT_UNREACH: return TEXT("SE_UDP_ERR_PORT_UNREACH");
+		case SE_ADDRFAMILY: return TEXT("SE_ADDRFAMILY");
+		case SE_SYSTEM: return TEXT("SE_SYSTEM");
+		case SE_NODEV: return TEXT("SE_NODEV");
 		default: return TEXT("Unknown Error");
 	};
-#else
-	return TEXT("");
-#endif
 }
 
 
