VoiceEngineImpl Update to handle Win Audio Device Changes

diff --git a/Engine/Plugins/Online/OnlineSubsystemUtils/Source/OnlineSubsystemUtils/OnlineSubsystemUtils.Build.cs b/Engine/Plugins/Online/OnlineSubsystemUtils/Source/OnlineSubsystemUtils/OnlineSubsystemUtils.Build.cs
--- a/Engine/Plugins/Online/OnlineSubsystemUtils/Source/OnlineSubsystemUtils/OnlineSubsystemUtils.Build.cs
+++ b/Engine/Plugins/Online/OnlineSubsystemUtils/Source/OnlineSubsystemUtils/OnlineSubsystemUtils.Build.cs
@@ -1,6 +1,7 @@
 // Copyright 1998-2019 Epic Games, Inc. All Rights Reserved.
 
 using UnrealBuildTool;
+using System.IO;
 
 public class OnlineSubsystemUtils : ModuleRules
 {
@@ -11,6 +12,19 @@ public class OnlineSubsystemUtils : ModuleRules
 
 		PrivateIncludePaths.Add("OnlineSubsystemUtils/Private");
 
+        string EnginePath = Path.GetFullPath(Target.RelativeEnginePath);
+        string RuntimePath = EnginePath + "Source/Runtime/";
+
+        bool bIsWindowsPlatformBuild = ((Target.Platform == UnrealTargetPlatform.Win32) || (Target.Platform == UnrealTargetPlatform.Win64));
+
+        if (bIsWindowsPlatformBuild)
+        {
+            AddEngineThirdPartyPrivateStaticDependencies(Target, "DX11Audio");
+            PrivateIncludePaths.Add(RuntimePath + "Windows/XAudio2/Public");
+            PrivateIncludePaths.Add(RuntimePath + "Windows/XAudio2/Private");
+        }
+
+
 		PrivateDependencyModuleNames.AddRange(
 			new string[] {
 				"ImageCore",
diff --git a/Engine/Plugins/Online/OnlineSubsystemUtils/Source/OnlineSubsystemUtils/Private/VoiceEngineImpl.cpp b/Engine/Plugins/Online/OnlineSubsystemUtils/Source/OnlineSubsystemUtils/Private/VoiceEngineImpl.cpp
--- a/Engine/Plugins/Online/OnlineSubsystemUtils/Source/OnlineSubsystemUtils/Private/VoiceEngineImpl.cpp
+++ b/Engine/Plugins/Online/OnlineSubsystemUtils/Source/OnlineSubsystemUtils/Private/VoiceEngineImpl.cpp
@@ -13,6 +13,14 @@
 /** Largest size allowed to carry over into next buffer */
 #define MAX_VOICE_REMAINDER_SIZE 4 * 1024
 
+#if PLATFORM_WINDOWS
+#include "XAudio2Support.h"
+namespace NotificationClient
+{
+	TSharedPtr<FMMNotificationClient> WindowsNotificationClient;
+}
+#endif 
+
 FRemoteTalkerDataImpl::FRemoteTalkerDataImpl() :
 	MaxUncompressedDataSize(0),
 	MaxUncompressedDataQueueSize(0),
@@ -102,6 +110,19 @@ void FRemoteTalkerDataImpl::Reset()
 	if (VoipSynthComponent)
 	{
 		VoipSynthComponent->Stop();
+
+		UAudioComponent* AudioComponent = VoipSynthComponent->GetAudioComponent();
+		if (AudioComponent->IsRegistered())
+		{
+			AudioComponent->UnregisterComponent();
+		}
+
+		//If the UVOIPTalker associated with this is still alive, notify it that this player is done talking.
+		if (UVOIPStatics::IsVOIPTalkerStillAlive(CachedTalkerPtr))
+		{
+			CachedTalkerPtr->OnTalkingEnd();
+		}
+
 		bIsActive = false;
 	}
 
@@ -134,7 +155,8 @@ FVoiceEngineImpl ::FVoiceEngineImpl() :
 	AvailableVoiceResult(EVoiceCaptureState::UnInitialized),
 	bPendingFinalCapture(false),
 	bIsCapturing(false),
-	SerializeHelper(nullptr)
+	SerializeHelper(nullptr),
+	bAudioDeviceChanged(false)
 {
 }
 
@@ -148,7 +170,8 @@ FVoiceEngineImpl::FVoiceEngineImpl(IOnlineSubsystem* InSubsystem) :
 	AvailableVoiceResult(EVoiceCaptureState::UnInitialized),
 	bPendingFinalCapture(false),
 	bIsCapturing(false),
-	SerializeHelper(nullptr)
+	SerializeHelper(nullptr),
+	bAudioDeviceChanged(false)
 {
 	FCoreUObjectDelegates::PostLoadMapWithWorld.AddRaw(this, &FVoiceEngineImpl::OnPostLoadMap);
 }
@@ -237,6 +260,9 @@ bool FVoiceEngineImpl::Init(int32 MaxLocalTalkers, int32 MaxRemoteTalkers)
 			bSuccess = VoiceCapture.IsValid() && VoiceEncoder.IsValid();
 			if (bSuccess)
 			{
+#if PLATFORM_WINDOWS
+				RegisterDeviceChangedListener();
+#endif
 				CompressedVoiceBuffer.Empty(UVOIPStatics::GetMaxCompressedVoiceDataSize());
 				DecompressedVoiceBuffer.Empty(UVOIPStatics::GetMaxUncompressedVoiceDataSizePerChannel());
 
@@ -564,23 +590,8 @@ void FVoiceEngineImpl::TickTalkers(float DeltaTime)
 		FRemoteTalkerDataImpl& RemoteData = It.Value();
 		double TimeSince = CurTime - RemoteData.LastSeen;
 
-		if (RemoteData.VoipSynthComponent->IsIdling() && RemoteData.bIsActive)
+		if (RemoteData.VoipSynthComponent && RemoteData.VoipSynthComponent->IsIdling() && RemoteData.bIsActive)
 		{
-			RemoteData.VoipSynthComponent->Stop();
-
-			UAudioComponent* AudioComponent = RemoteData.VoipSynthComponent->GetAudioComponent();
-			if (AudioComponent->IsRegistered())
-			{
-				AudioComponent->UnregisterComponent();
-			}
-
-			//If the UVOIPTalker associated with this is still alive, notify it that this player is done talking.
-			if (UVOIPStatics::IsVOIPTalkerStillAlive(RemoteData.CachedTalkerPtr))
-			{
-				RemoteData.CachedTalkerPtr->OnTalkingEnd();
-			}
-			RemoteData.bIsActive = false;
-
 			RemoteData.Reset();
 		}
 		else if (TimeSince >= UVOIPStatics::GetRemoteTalkerTimeoutDuration())
@@ -600,6 +611,13 @@ void FVoiceEngineImpl::Tick(float DeltaTime)
 	}
 
 	TickTalkers(DeltaTime);
+
+#if PLATFORM_WINDOWS
+	if (bAudioDeviceChanged)
+	{
+		HandleDeviceChange();
+	}
+#endif
 }
 
 void FVoiceEngineImpl::GenerateVoiceData(USoundWaveProcedural* InProceduralWave, int32 SamplesRequired, const FUniqueNetId& TalkerId)
@@ -634,7 +652,7 @@ void FVoiceEngineImpl::OnAudioFinished()
 	for (FRemoteTalkerData::TIterator It(RemoteTalkerBuffers); It; ++It)
 	{
 		FRemoteTalkerDataImpl& RemoteData = It.Value();
-		if (RemoteData.VoipSynthComponent->IsIdling())
+		if (RemoteData.VoipSynthComponent && RemoteData.VoipSynthComponent->IsIdling())
 		{
 			UE_LOG_ONLINE_VOICEENGINE(Log, TEXT("Removing VOIP AudioComponent for Id: %s"), *It.Key().ToDebugString());
 			RemoteData.VoipSynthComponent->Stop();
@@ -765,4 +783,51 @@ void FVoiceEngineImpl::CreateSerializeHelper()
 	{
 		SerializeHelper = new FVoiceSerializeHelper(this);
 	}
-}
\ No newline at end of file
+}
+
+#if PLATFORM_WINDOWS
+void FVoiceEngineImpl::RegisterDeviceChangedListener()
+{
+	if (!NotificationClient::WindowsNotificationClient.IsValid())
+	{
+		NotificationClient::WindowsNotificationClient = TSharedPtr<FMMNotificationClient>(new FMMNotificationClient);
+	}
+
+	NotificationClient::WindowsNotificationClient->RegisterDeviceChangedListener(this);
+}
+
+void FVoiceEngineImpl::UnregisterDeviceChangedListener()
+{
+	if (NotificationClient::WindowsNotificationClient.IsValid())
+	{
+		NotificationClient::WindowsNotificationClient->UnRegisterDeviceDeviceChangedListener(this);
+	}
+}
+
+void FVoiceEngineImpl::HandleDeviceChange()
+{
+	const double TimeSince = FPlatformTime::Seconds() - TimeDeviceChaned;
+	if (TimeSince >= DeviceChangeDelay)
+	{
+		if (bIsCapturing)
+		{
+			StopLocalVoiceProcessing(OwningUserIndex);
+			StartLocalVoiceProcessing(OwningUserIndex);
+		}
+
+		for (FRemoteTalkerData::TIterator It(RemoteTalkerBuffers); It; ++It)
+		{
+			FRemoteTalkerDataImpl& RemoteData = It.Value();
+			RemoteData.Reset();
+		}
+
+		bAudioDeviceChanged = false;
+	}
+}
+
+void FVoiceEngineImpl::OnDefaultDeviceChanged()
+{
+	bAudioDeviceChanged = true;
+	TimeDeviceChaned = FPlatformTime::Seconds();
+}
+#endif
\ No newline at end of file
diff --git a/Engine/Plugins/Online/OnlineSubsystemUtils/Source/OnlineSubsystemUtils/Public/VoiceEngineImpl.h b/Engine/Plugins/Online/OnlineSubsystemUtils/Source/OnlineSubsystemUtils/Public/VoiceEngineImpl.h
--- a/Engine/Plugins/Online/OnlineSubsystemUtils/Source/OnlineSubsystemUtils/Public/VoiceEngineImpl.h
+++ b/Engine/Plugins/Online/OnlineSubsystemUtils/Source/OnlineSubsystemUtils/Public/VoiceEngineImpl.h
@@ -13,6 +13,7 @@
 #include "OnlineSubsystemUtilsPackage.h"
 #include "VoipListenerSynthComponent.h"
 #include "OnlineSubsystemUtilsPackage.h"
+#include "AudioDevice.h"
 
 #include "VoicePacketImpl.h"
 #include "UObject/CoreOnline.h"
@@ -85,7 +86,7 @@ public:
 /**
  * Generic implementation of voice engine, using Voice module for capture/codec
  */
-class ONLINESUBSYSTEMUTILS_API FVoiceEngineImpl : public IVoiceEngine, public FSelfRegisteringExec
+class ONLINESUBSYSTEMUTILS_API FVoiceEngineImpl : public IVoiceEngine, public FSelfRegisteringExec, public IDeviceChangedListener
 {
 	class FVoiceSerializeHelper : public FGCObject
 	{
@@ -153,6 +154,13 @@ class ONLINESUBSYSTEMUTILS_API FVoiceEngineImpl : public IVoiceEngine, public FS
 	/** Serialization helper */
 	FVoiceSerializeHelper* SerializeHelper;
 
+// Get Audio Device Changes on Windows
+#if PLATFORM_WINDOWS
+	FThreadSafeBool bAudioDeviceChanged;
+	double TimeDeviceChaned;
+	const double DeviceChangeDelay = 2.0f;
+#endif
+
 	/**
 	 * Determines if the specified index is the owner or not
 	 *
@@ -299,6 +307,18 @@ protected:
 	FLocalVoiceData*           GetLocalPlayerVoiceData()    { return PlayerVoiceData; }
 	int32                      GetMaxVoiceRemainderSize();
 	void					   CreateSerializeHelper();
+
+	// Get Audio Device Changes on Windows
+#if PLATFORM_WINDOWS
+	virtual void RegisterDeviceChangedListener();
+	virtual void UnregisterDeviceChangedListener();
+	virtual void HandleDeviceChange();
+
+	//~ Begin IDeviceChangedListener
+	virtual void OnDefaultDeviceChanged() override;
+	virtual void OnDeviceRemoved(FString DeviceID) override {}
+	//~ End IDeviceChangedListener
+#endif
 };
 
 typedef TSharedPtr<FVoiceEngineImpl, ESPMode::ThreadSafe> FVoiceEngineImplPtr;
diff --git a/Engine/Source/Runtime/Windows/XAudio2/Private/XAudio2Support.h b/Engine/Source/Runtime/Windows/XAudio2/Private/XAudio2Support.h
--- a/Engine/Source/Runtime/Windows/XAudio2/Private/XAudio2Support.h
+++ b/Engine/Source/Runtime/Windows/XAudio2/Private/XAudio2Support.h
@@ -32,6 +32,7 @@
 	#include <X3Daudio.h>
 #include "Windows/HideWindowsPlatformAtomics.h"
 #include "Windows/HideWindowsPlatformTypes.h"
+#include "XAudio2Device.h"
 
 #if PLATFORM_WINDOWS
 #include "Windows/AllowWindowsPlatformTypes.h"
@@ -199,6 +200,7 @@ struct FXWMABufferInfo
 	UINT32						XWMASeekDataSize;
 };
 
+class FXAudio2SoundBuffer;
 typedef FAsyncRealtimeAudioTaskProxy<FXAudio2SoundBuffer> FAsyncRealtimeAudioTask;
 
 /**
diff --git a/Engine/Source/Runtime/Windows/XAudio2/Public/XAudio2Device.h b/Engine/Source/Runtime/Windows/XAudio2/Public/XAudio2Device.h
--- a/Engine/Source/Runtime/Windows/XAudio2/Public/XAudio2Device.h
+++ b/Engine/Source/Runtime/Windows/XAudio2/Public/XAudio2Device.h
@@ -17,6 +17,7 @@
 #include "AudioDecompress.h"
 #include "AudioEffect.h"
 #include "AudioDevice.h"
+#include "VorbisAudioInfo.h"
 
 DECLARE_LOG_CATEGORY_EXTERN(LogXAudio2, Log, All);
 
