Fix Windows CIS

diff --git a/Engine/Plugins/Online/OnlineSubsystemUtils/Source/OnlineSubsystemUtils/Private/VoiceEngineImpl.cpp b/Engine/Plugins/Online/OnlineSubsystemUtils/Source/OnlineSubsystemUtils/Private/VoiceEngineImpl.cpp
--- a/Engine/Plugins/Online/OnlineSubsystemUtils/Source/OnlineSubsystemUtils/Private/VoiceEngineImpl.cpp
+++ b/Engine/Plugins/Online/OnlineSubsystemUtils/Source/OnlineSubsystemUtils/Private/VoiceEngineImpl.cpp
@@ -15,7 +15,6 @@
 
 #if PLATFORM_WINDOWS
 #include "XAudio2Support.h"
-HMODULE FXAudioDeviceProperties::XAudio2Dll;
 namespace NotificationClient
 {
 	TSharedPtr<FMMNotificationClient> WindowsNotificationClient;
