Added log to show configured vs OS-actual values when setting a Socket's buffer size.

diff --git a/Engine/Plugins/Online/OnlineSubsystemUtils/Source/OnlineSubsystemUtils/Private/IpNetDriver.cpp b/Engine/Plugins/Online/OnlineSubsystemUtils/Source/OnlineSubsystemUtils/Private/IpNetDriver.cpp
--- a/Engine/Plugins/Online/OnlineSubsystemUtils/Source/OnlineSubsystemUtils/Private/IpNetDriver.cpp
+++ b/Engine/Plugins/Online/OnlineSubsystemUtils/Source/OnlineSubsystemUtils/Private/IpNetDriver.cpp
@@ -158,11 +158,14 @@ bool UIpNetDriver::InitBase( bool bInitAsClient, FNetworkNotify* InNotify, const
 
 	// Increase socket queue size, because we are polling rather than threading
 	// and thus we rely on the OS socket to buffer a lot of data.
-	int32 RecvSize = bInitAsClient ? ClientDesiredSocketReceiveBufferBytes	: ServerDesiredSocketReceiveBufferBytes;
-	int32 SendSize = bInitAsClient ? ClientDesiredSocketSendBufferBytes		: ServerDesiredSocketSendBufferBytes;
-	Socket->SetReceiveBufferSize(RecvSize,RecvSize);
-	Socket->SetSendBufferSize(SendSize,SendSize);
-	UE_LOG(LogInit, Log, TEXT("%s: Socket queue %i / %i"), SocketSubsystem->GetSocketAPIName(), RecvSize, SendSize );
+	const int32 DesiredRecvSize = bInitAsClient ? ClientDesiredSocketReceiveBufferBytes	: ServerDesiredSocketReceiveBufferBytes;
+	const int32 DesiredSendSize = bInitAsClient ? ClientDesiredSocketSendBufferBytes	: ServerDesiredSocketSendBufferBytes;
+	int32 ActualRecvSize(0);
+	int32 ActualSendSize(0);
+	Socket->SetReceiveBufferSize(DesiredRecvSize, ActualRecvSize);
+	Socket->SetSendBufferSize(DesiredSendSize, ActualSendSize);
+	UE_LOG(LogInit, Log, TEXT("%s: Socket queue. Rx: %i (config %i) Tx: %i (config %i)"), SocketSubsystem->GetSocketAPIName(),
+		ActualRecvSize, DesiredRecvSize, ActualSendSize, DesiredSendSize);
 
 	// allow multihome as a url option
 	const TCHAR* MultiHomeBindAddr = URL.GetOption(TEXT("multihome="), nullptr);
