Adding backwards compatibility to the Voice getters.

#rnx

diff --git a/Engine/Plugins/Online/OnlineSubsystemUtils/Source/OnlineSubsystemUtils/Public/VoiceEngineImpl.h b/Engine/Plugins/Online/OnlineSubsystemUtils/Source/OnlineSubsystemUtils/Public/VoiceEngineImpl.h
--- a/Engine/Plugins/Online/OnlineSubsystemUtils/Source/OnlineSubsystemUtils/Public/VoiceEngineImpl.h
+++ b/Engine/Plugins/Online/OnlineSubsystemUtils/Source/OnlineSubsystemUtils/Public/VoiceEngineImpl.h
@@ -302,7 +302,9 @@ private:
 protected:
 	virtual IOnlineSubsystem*				 GetOnlineSubSystem()			{ return OnlineSubsystem; }
 	virtual const TSharedPtr<IVoiceCapture>& GetVoiceCapture() const		{ return VoiceCapture; }
+	virtual TSharedPtr<IVoiceCapture>&		 GetVoiceCapture()				{ return VoiceCapture; }
 	virtual const TSharedPtr<IVoiceEncoder>& GetVoiceEncoder() const		{ return VoiceEncoder; }
+	virtual TSharedPtr<IVoiceEncoder>&		 GetVoiceEncoder()				{ return VoiceEncoder; }
 	virtual FRemoteTalkerData&				 GetRemoteTalkerBuffers()		{ return RemoteTalkerBuffers; }
 	virtual TArray<uint8>&					 GetCompressedVoiceBuffer()		{ return CompressedVoiceBuffer; }
 	virtual TArray<uint8>&					 GetDecompressedVoiceBuffer()	{ return DecompressedVoiceBuffer; }
