Beacon Party Reservations now keep track of what players have been removed from their reservation.
Use case: in Tournaments for scoring purposes

diff --git a/Engine/Plugins/Online/OnlineSubsystemUtils/Source/OnlineSubsystemUtils/Private/PartyBeaconHost.cpp b/Engine/Plugins/Online/OnlineSubsystemUtils/Source/OnlineSubsystemUtils/Private/PartyBeaconHost.cpp
--- a/Engine/Plugins/Online/OnlineSubsystemUtils/Source/OnlineSubsystemUtils/Private/PartyBeaconHost.cpp
+++ b/Engine/Plugins/Online/OnlineSubsystemUtils/Source/OnlineSubsystemUtils/Private/PartyBeaconHost.cpp
@@ -743,10 +743,7 @@ EPartyReservationResult::Type APartyBeaconHost::UpdatePartyReservation(const FPa
 										{
 											FormerReservation.Dump();
 										}
-										int32 NumReservationsRemoved = FormerReservation.PartyMembers.RemoveAll([&PlayerRes](const FPlayerReservation& OtherRes)
-										{
-											return OtherRes.UniqueId == PlayerRes.UniqueId;
-										});
+										int32 NumReservationsRemoved = FormerReservation.RemoveAllPartyMembers(PlayerRes);
 
 										State->NumConsumedReservations -= NumReservationsRemoved;
 										UE_LOG(LogPartyBeacon, Verbose, TEXT("APartyBeaconHost::UpdatePartyReservation: Removed %d players, setting NumConsumedReservations to %d"), NumReservationsRemoved, State->NumConsumedReservations);
diff --git a/Engine/Plugins/Online/OnlineSubsystemUtils/Source/OnlineSubsystemUtils/Private/PartyBeaconState.cpp b/Engine/Plugins/Online/OnlineSubsystemUtils/Source/OnlineSubsystemUtils/Private/PartyBeaconState.cpp
--- a/Engine/Plugins/Online/OnlineSubsystemUtils/Source/OnlineSubsystemUtils/Private/PartyBeaconState.cpp
+++ b/Engine/Plugins/Online/OnlineSubsystemUtils/Source/OnlineSubsystemUtils/Private/PartyBeaconState.cpp
@@ -85,6 +85,30 @@ bool FPartyReservation::CanPlayerMigrateFromReservation(const FPartyReservation&
 	return TeamNum == Other.TeamNum;
 }
 
+int32 FPartyReservation::RemoveAllPartyMembers(const FPlayerReservation& OtherRes)
+{
+	int32 NumRemoved = 0;
+	for(int32 Idx = PartyMembers.Num() - 1; Idx >= 0; Idx--)
+	{
+		if (PartyMembers[Idx].UniqueId == OtherRes.UniqueId)
+		{
+			NumRemoved++;
+			RemovePartyMemberAtIndex(Idx);
+		}
+	}
+
+	return NumRemoved;
+}
+
+void FPartyReservation::RemovePartyMemberAtIndex(int32 Idx)
+{
+	if (-1 < Idx && Idx < PartyMembers.Num())
+	{
+		RemovedPartyMembers.Add(PartyMembers[Idx]);
+		PartyMembers.RemoveAtSwap(Idx);
+	}
+}
+
 UPartyBeaconState::UPartyBeaconState(const FObjectInitializer& ObjectInitializer) :
 	Super(ObjectInitializer),
 	SessionName(NAME_None),
@@ -714,7 +738,7 @@ void UPartyBeaconState::UpdatePartyLeader(const FUniqueNetIdRepl& InPartyMemberI
 						FPlayerReservation PlayerReservation = PriorReservation->PartyMembers[PriorPlayerReservationIdx];
 
 						// Remove player from their previous reservation and find a new place for them
-						PriorReservation->PartyMembers.RemoveAtSwap(PriorPlayerReservationIdx);
+						PriorReservation->RemovePartyMemberAtIndex(PriorPlayerReservationIdx);
 
 						// If there is already a reservation that has the new party leader as a leader, join it
 						// If not, create one
@@ -890,7 +914,7 @@ bool UPartyBeaconState::RemovePlayer(const FUniqueNetIdRepl& PlayerId)
 			if (PlayerEntry.UniqueId == PlayerId)
 			{
 				// player removed
-				Reservation.PartyMembers.RemoveAtSwap(PlayerIdx--);
+				Reservation.RemovePartyMemberAtIndex(PlayerIdx--);
 				bWasRemoved = true;
 
 				// free up a consumed entry
diff --git a/Engine/Plugins/Online/OnlineSubsystemUtils/Source/OnlineSubsystemUtils/Public/PartyBeaconState.h b/Engine/Plugins/Online/OnlineSubsystemUtils/Source/OnlineSubsystemUtils/Public/PartyBeaconState.h
--- a/Engine/Plugins/Online/OnlineSubsystemUtils/Source/OnlineSubsystemUtils/Public/PartyBeaconState.h
+++ b/Engine/Plugins/Online/OnlineSubsystemUtils/Source/OnlineSubsystemUtils/Public/PartyBeaconState.h
@@ -195,6 +195,12 @@ struct ONLINESUBSYSTEMUTILS_API FPartyReservation
 	UPROPERTY(Transient)
 	TArray<FPlayerReservation> PartyMembers;
 
+	/** Keeping a record of all logged out players from this reservation. */
+	UPROPERTY(Transient)
+	TArray<FPlayerReservation> RemovedPartyMembers;
+	void RemovePartyMemberAtIndex(int32 Idx);
+	int32 RemoveAllPartyMembers(const FPlayerReservation& OtherRes);
+
 	/** Is this data well formed */
 	bool IsValid() const;
 
