Fix crashing that occurs when an invalid friend id is provided to the joinfriend command. Additionally, attempt to parse the input as a platform id if the raw string of bytes fails.

diff --git a/Engine/Plugins/Online/OnlineSubsystemUtils/Source/OnlineSubsystemUtils/Private/Tests/TestSessionInterface.cpp b/Engine/Plugins/Online/OnlineSubsystemUtils/Source/OnlineSubsystemUtils/Private/Tests/TestSessionInterface.cpp
--- a/Engine/Plugins/Online/OnlineSubsystemUtils/Source/OnlineSubsystemUtils/Private/Tests/TestSessionInterface.cpp
+++ b/Engine/Plugins/Online/OnlineSubsystemUtils/Source/OnlineSubsystemUtils/Private/Tests/TestSessionInterface.cpp
@@ -511,8 +511,16 @@ bool FTestSessionInterface::Exec(UWorld* InWorld, const TCHAR* Cmd, FOutputDevic
 				if (FParse::Token(Cmd, FriendIdStr, ARRAY_COUNT(FriendIdStr), true))
 				{
 					TSharedPtr<const FUniqueNetId> FriendId = Identity->CreateUniquePlayerId((uint8*)FriendIdStr, FCString::Strlen(FriendIdStr));
-					OnFindFriendSessionCompleteDelegateHandles.Add(LocalUserNum, SessionInt->AddOnFindFriendSessionCompleteDelegate_Handle(LocalUserNum, OnFindFriendSessionCompleteDelegate));
-					SessionInt->FindFriendSession(LocalUserNum, *FriendId);
+					if (!FriendId.IsValid())
+					{
+						FriendId = Identity->CreateUniquePlayerId(FriendIdStr);
+					}
+
+					if (FriendId.IsValid())
+					{
+						OnFindFriendSessionCompleteDelegateHandles.Add(LocalUserNum, SessionInt->AddOnFindFriendSessionCompleteDelegate_Handle(LocalUserNum, OnFindFriendSessionCompleteDelegate));
+						SessionInt->FindFriendSession(LocalUserNum, *FriendId);
+					}
 				}
 			}
 			bWasHandled = true;
