This changelist includes several fixes for issues related to or reported in the above ticket.

Fixes a case where a party/leader os rejected attempting to join a creative server they do not have a reservation as a leader on, but that one or more of the party members had an exiting reservation.
The code now checks that the existing reservations for those players aren't incompatible allowing the reservation only if all the members with existing reservations on the same team, and can the new members fit on that team.
(can be caused if that members fails to flush into the server on a previous join)

Fixes a case where a change of party leader would result in broken reservations being left on the beacon. 2 reservations with the same leader, and one in one of those the leader would not be a member.
This would create a situation where that leaser, or (baring the above fix) anyone in a party with them wouldn't be able to join for the lifetime of the server.

Fixes a case where party leaders who have a reservation on the server would be rejected attempting to rejoin (or join the case of a player in the lobby having been promoted to lead) unless their new reservation exactly matched the previous one in terms of party members (which never happens for a leader JIP). The code now also allows them to rejoin if the newly requested reservation is a subset of their existing reservation.

diff --git a/Engine/Plugins/Online/OnlineSubsystemUtils/Source/OnlineSubsystemUtils/Private/PartyBeaconHost.cpp b/Engine/Plugins/Online/OnlineSubsystemUtils/Source/OnlineSubsystemUtils/Private/PartyBeaconHost.cpp
--- a/Engine/Plugins/Online/OnlineSubsystemUtils/Source/OnlineSubsystemUtils/Private/PartyBeaconHost.cpp
+++ b/Engine/Plugins/Online/OnlineSubsystemUtils/Source/OnlineSubsystemUtils/Private/PartyBeaconHost.cpp
@@ -507,9 +507,9 @@ EPartyReservationResult::Type APartyBeaconHost::AddPartyReservation(const FParty
 		if (ExistingReservationIdx != INDEX_NONE)
 		{
 			FPartyReservation& ExistingReservation = Reservations[ExistingReservationIdx];
-			if (ReservationRequest.PartyMembers.Num() == ExistingReservation.PartyMembers.Num())
+			if (ReservationRequest.PartyMembers.Num() <= ExistingReservation.PartyMembers.Num())
 			{
-				// Verify the reservations are the same
+				// Verify the reservations are either the same, or that the incoming reservation is a subset of the existing one
 				int32 NumMatchingReservations = 0;
 				for (const FPlayerReservation& NewPlayerRes : ReservationRequest.PartyMembers)
 				{
@@ -525,7 +525,7 @@ EPartyReservationResult::Type APartyBeaconHost::AddPartyReservation(const FParty
 					}
 				}
 
-				if (NumMatchingReservations == ExistingReservation.PartyMembers.Num())
+				if (NumMatchingReservations == ReservationRequest.PartyMembers.Num())
 				{
 					for (const FPlayerReservation& NewPlayerRes : ReservationRequest.PartyMembers)
 					{
@@ -574,7 +574,7 @@ EPartyReservationResult::Type APartyBeaconHost::AddPartyReservation(const FParty
 			}
 			else
 			{
-				// Existing reservation doesn't match incoming duplicate reservation
+				// Existing reservation holds less players than the incoming duplicate reservation
 				Result = EPartyReservationResult::IncorrectPlayerCount;
 			}
 		}
@@ -582,18 +582,34 @@ EPartyReservationResult::Type APartyBeaconHost::AddPartyReservation(const FParty
 		{
 			// Check for players we already have reservations for
 			// Keep track of team index for existing members - if we have members on opposing teams, reject this reservation
-			bool bContainsExistingMembers = false;
+			bool bContainsIncompatibleExistingMembers = false;
 			int32 ExistingMemberReservationTeamNum = INDEX_NONE;
+			int32 TeamIdx = -1;
+			int32 NumTeamPlayersWithExistingReservation = 0;
+			TArray<FPlayerReservation> ExistingPartyMemberReservations;
 			for (const FPlayerReservation& PartyMember : ReservationRequest.PartyMembers)
 			{
 				int32 MemberExistingPartyReservationIdx = State->GetExistingReservationContainingMember(PartyMember.UniqueId);
+				
 				if (MemberExistingPartyReservationIdx != INDEX_NONE)
 				{
-					const FPartyReservation& MemberExistingPartyReservation = Reservations[MemberExistingPartyReservationIdx];
-					UE_LOG(LogPartyBeacon, Display, TEXT("APartyBeaconHost::AddPartyReservation: Found existing reservation for party member %s"), *PartyMember.UniqueId.ToDebugString());
-					ReservationRequest.Dump();
-					MemberExistingPartyReservation.Dump();
-					bContainsExistingMembers = true;
+					ExistingPartyMemberReservations.Add(PartyMember);
+					const FPartyReservation& MemberExistingPartyReservation = Reservations[MemberExistingPartyReservationIdx];		
+					// is the member on a suitable team (same as previous members/we the part doesn't have a team yet)
+					if (TeamIdx == -1 || TeamIdx == MemberExistingPartyReservation.TeamNum) 
+					{
+						TeamIdx = MemberExistingPartyReservation.TeamNum;
+						NumTeamPlayersWithExistingReservation++;
+					}
+					else
+					{
+						UE_LOG(LogPartyBeacon, Display, TEXT("APartyBeaconHost::AddPartyReservation: Found existing reservation with missmatched team for party member %s"), *PartyMember.UniqueId.ToDebugString());
+						ReservationRequest.Dump();
+						MemberExistingPartyReservation.Dump();
+						bContainsIncompatibleExistingMembers = true;
+					}
+				
+					
 				}
 				else
 				{
@@ -604,12 +620,30 @@ EPartyReservationResult::Type APartyBeaconHost::AddPartyReservation(const FParty
 					{
 						UE_LOG(LogPartyBeacon, Display, TEXT("APartyBeaconHost::AddPartyReservation: Found party member %s in the pending player list"), *PartyMember.UniqueId.ToDebugString());
 						ReservationRequest.Dump();
-						bContainsExistingMembers = true;
+						bContainsIncompatibleExistingMembers = true;
 					}
 				}
+			}	
+			if (TeamIdx != -1)
+			{	
+				// Will the part members not currently on the team fit in it, and fit within the server
+				const int32 NumTeamMembers = GetNumPlayersOnTeam(TeamIdx);
+				const int32 NumAvailableSlotsOnTeam = FMath::Max<int32>(0, (GetMaxPlayersPerTeam() - NumTeamMembers) + NumTeamPlayersWithExistingReservation);
+				if ((NumAvailableSlotsOnTeam < ReservationRequest.PartyMembers.Num()) || (State->GetRemainingReservations() < (ReservationRequest.PartyMembers.Num() - NumTeamPlayersWithExistingReservation)))
+				{
+					bContainsIncompatibleExistingMembers = true;
+				}
+			}
+			// if this party reservation included players who already had reservations, but is not incompatible based on teams then remove those players previous reservations
+			if (!bContainsIncompatibleExistingMembers)
+			{
+				for (const FPlayerReservation& MemberReservation : ExistingPartyMemberReservations)
+				{					
+					State->RemovePlayer(MemberReservation.UniqueId);
+				}
 			}
 
-			if (!bContainsExistingMembers)
+			if (!bContainsIncompatibleExistingMembers)
 			{
 				if (State->DoesReservationFit(ReservationRequest))
 				{
@@ -776,8 +810,7 @@ EPartyReservationResult::Type APartyBeaconHost::UpdatePartyReservation(const FPa
 		}
 		else
 		{
-			UE_LOG(LogPartyBeacon, Verbose, TEXT("Adding Member"));
-			
+			UE_LOG(LogPartyBeacon, Verbose, TEXT("Adding Member"));			
 			const int32 ExistingReservationIdx = State->GetExistingReservation(ReservationUpdateRequest.PartyLeader);
 			if (ExistingReservationIdx != INDEX_NONE)
 			{
@@ -946,6 +979,7 @@ EPartyReservationResult::Type APartyBeaconHost::UpdatePartyReservation(const FPa
 				// Send a not found reservation response
 				Result = EPartyReservationResult::ReservationNotFound;
 			}
+
 		}
 	}
 	else
diff --git a/Engine/Plugins/Online/OnlineSubsystemUtils/Source/OnlineSubsystemUtils/Private/PartyBeaconState.cpp b/Engine/Plugins/Online/OnlineSubsystemUtils/Source/OnlineSubsystemUtils/Private/PartyBeaconState.cpp
--- a/Engine/Plugins/Online/OnlineSubsystemUtils/Source/OnlineSubsystemUtils/Private/PartyBeaconState.cpp
+++ b/Engine/Plugins/Online/OnlineSubsystemUtils/Source/OnlineSubsystemUtils/Private/PartyBeaconState.cpp
@@ -716,83 +716,98 @@ void UPartyBeaconState::UpdatePartyLeader(const FUniqueNetIdRepl& InPartyMemberI
 		// PartyMemberReservationIdx may be invalid if we had just started kicking this player
 		if (PartyMemberReservationIdx != INDEX_NONE)
 		{
-			// Get the reservation that has NewPartyLeaderId as the leader
-			// May be INDEX_NONE if the target is not the leader of a reservation and is instead just a member
+			int32 PartLeaderAsMemberReservationIdx = GetExistingReservationContainingMember(NewPartyLeaderId);	
 			int32 NewPartyLeaderReservationIdx = GetExistingReservation(NewPartyLeaderId);
-			if (NewPartyLeaderReservationIdx != PartyMemberReservationIdx)
+			bool Promotion = false;
+			// if the leader we want to be under is in the same party as us, and we are the current leader then promote them to the leader of the existing party. Excluding cases where we are asking to be out own leader
+			if (PartyMemberReservationIdx == PartLeaderAsMemberReservationIdx && (InPartyMemberId != NewPartyLeaderId))
 			{
-				FPartyReservation* PriorReservation = &Reservations[PartyMemberReservationIdx];
-				FPartyReservation* DestinationReservation = (NewPartyLeaderReservationIdx != INDEX_NONE) ? &Reservations[NewPartyLeaderReservationIdx] : nullptr;
-
-				// Verify that a migration between existing reservations can occur
-				if (!DestinationReservation || DestinationReservation->CanPlayerMigrateFromReservation(*PriorReservation))
+				FPartyReservation* SharedReservation = &Reservations[PartLeaderAsMemberReservationIdx];
+				if (SharedReservation->PartyLeader == InPartyMemberId)
 				{
-					// Make a copy of their player reservation so we can move it to the new reservation
-					int32 PriorPlayerReservationIdx = PriorReservation->PartyMembers.IndexOfByPredicate([&InPartyMemberId](const FPlayerReservation& ExistingPlayerRes)
-					{
-						return InPartyMemberId == ExistingPlayerRes.UniqueId;
-					});
-					if (ensure(PriorPlayerReservationIdx != INDEX_NONE))
-					{
-						// Make a copy of the player reservation to be inserted into a new reservation after being removed from the prior reservation
-						FPlayerReservation PlayerReservation = PriorReservation->PartyMembers[PriorPlayerReservationIdx];
-
-						// Remove player from their previous reservation and find a new place for them
-						PriorReservation->RemovePartyMemberAtIndex(PriorPlayerReservationIdx);
+					SharedReservation->PartyLeader = NewPartyLeaderId;
+					Promotion = true;
+				}								
+			}
+			if (Promotion == false)
+			{
+				// Get the reservation that has NewPartyLeaderId as the leader
+				// May be INDEX_NONE if the target is not the leader of a reservation and is instead just a member
+				if (NewPartyLeaderReservationIdx != PartyMemberReservationIdx)
+				{
+					FPartyReservation* PriorReservation = &Reservations[PartyMemberReservationIdx];
+					FPartyReservation* DestinationReservation = (NewPartyLeaderReservationIdx != INDEX_NONE) ? &Reservations[NewPartyLeaderReservationIdx] : nullptr;
 
-						// If there is already a reservation that has the new party leader as a leader, join it
-						// If not, create one
-						if (DestinationReservation)
+					// Verify that a migration between existing reservations can occur
+					if (!DestinationReservation || DestinationReservation->CanPlayerMigrateFromReservation(*PriorReservation))
+					{
+						// Make a copy of their player reservation so we can move it to the new reservation
+						int32 PriorPlayerReservationIdx = PriorReservation->PartyMembers.IndexOfByPredicate([&InPartyMemberId](const FPlayerReservation& ExistingPlayerRes)
 						{
-							UE_LOG(LogPartyBeacon, Verbose, TEXT("UpdatePartyLeader:  Moving player %s from reservation under party leader %s, to reservation under party leader %s"),
-								*InPartyMemberId.ToString(), *PriorReservation->PartyLeader.ToString(), *NewPartyLeaderId.ToString());
-							DestinationReservation->PartyMembers.Add(MoveTemp(PlayerReservation));
-						}
-						else
+							return InPartyMemberId == ExistingPlayerRes.UniqueId;
+						});
+						if (ensure(PriorPlayerReservationIdx != INDEX_NONE))
 						{
-							UE_LOG(LogPartyBeacon, Verbose, TEXT("UpdatePartyLeader:  Moving player %s from reservation under party leader %s, to new reservation with leader %s"),
-								*InPartyMemberId.ToString(), *PriorReservation->PartyLeader.ToString(), *NewPartyLeaderId.ToString());
+							// Make a copy of the player reservation to be inserted into a new reservation after being removed from the prior reservation
+							FPlayerReservation PlayerReservation = PriorReservation->PartyMembers[PriorPlayerReservationIdx];
+
+							// Remove player from their previous reservation and find a new place for them
+							PriorReservation->RemovePartyMemberAtIndex(PriorPlayerReservationIdx);
 
+							// If there is already a reservation that has the new party leader as a leader, join it
+							// If not, create one
+							if (DestinationReservation)
+							{
+								UE_LOG(LogPartyBeacon, Verbose, TEXT("UpdatePartyLeader:  Moving player %s from reservation under party leader %s, to reservation under party leader %s"),
+									*InPartyMemberId.ToString(), *PriorReservation->PartyLeader.ToString(), *NewPartyLeaderId.ToString());
+								DestinationReservation->PartyMembers.Add(MoveTemp(PlayerReservation));
+							}
+							else
 							{
-								FPartyReservation NewReservation;
-								NewReservation.TeamNum = PriorReservation->TeamNum;
-								NewReservation.PartyLeader = NewPartyLeaderId;
-								NewReservation.PartyMembers.Add(MoveTemp(PlayerReservation));
+								UE_LOG(LogPartyBeacon, Verbose, TEXT("UpdatePartyLeader:  Moving player %s from reservation under party leader %s, to new reservation with leader %s"),
+									*InPartyMemberId.ToString(), *PriorReservation->PartyLeader.ToString(), *NewPartyLeaderId.ToString());
+
+								{
+									FPartyReservation NewReservation;
+									NewReservation.TeamNum = PriorReservation->TeamNum;
+									NewReservation.PartyLeader = NewPartyLeaderId;
+									NewReservation.PartyMembers.Add(MoveTemp(PlayerReservation));
+
+									Reservations.Add(MoveTemp(NewReservation));
+								}
+								// While the index doesn't change, the Add() call above could invalidate our memory address
+								PriorReservation = &Reservations[PartyMemberReservationIdx];
+							}
 
-								Reservations.Add(MoveTemp(NewReservation));
+							// If the former reservation is now empty, remove the reservation
+							if (PriorReservation->PartyMembers.Num() == 0)
+							{
+								UE_LOG(LogPartyBeacon, Verbose, TEXT("UpdatePartyLeader:  Removing now empty reservation that had party leader %s"), *PriorReservation->PartyLeader.ToString());
+								Reservations.RemoveAtSwap(PartyMemberReservationIdx);
+								PriorReservation = nullptr;
 							}
-							// While the index doesn't change, the Add() call above could invalidate our memory address
-							PriorReservation = &Reservations[PartyMemberReservationIdx];
-						}
 
-						// If the former reservation is now empty, remove the reservation
-						if (PriorReservation->PartyMembers.Num() == 0)
+							// At this point both PriorReservation and DestinationReservation are possibly invalid, don't use them below here
+
+							SanityCheckReservations(false);
+						}
+						else
 						{
-							UE_LOG(LogPartyBeacon, Verbose, TEXT("UpdatePartyLeader:  Removing now empty reservation that had party leader %s"), *PriorReservation->PartyLeader.ToString());
-							Reservations.RemoveAtSwap(PartyMemberReservationIdx);
-							PriorReservation = nullptr;
+							UE_LOG(LogPartyBeacon, Warning, TEXT("UpdatePartyLeader:  Member %s not found in their own reservation!"), *InPartyMemberId.ToDebugString());
 						}
-
-						// At this point both PriorReservation and DestinationReservation are possibly invalid, don't use them below here
-
-						SanityCheckReservations(false);
 					}
 					else
 					{
-						UE_LOG(LogPartyBeacon, Warning, TEXT("UpdatePartyLeader:  Member %s not found in their own reservation!"), *InPartyMemberId.ToDebugString());
+						UE_LOG(LogPartyBeacon, Warning, TEXT("UpdatePartyLeader:  Unable to migrate player %s from reservation under leader %s to existing reservation with leader %s"),
+							*InPartyMemberId.ToString(), *PriorReservation->PartyLeader.ToString(), *NewPartyLeaderId.ToString());
 					}
 				}
 				else
 				{
-					UE_LOG(LogPartyBeacon, Warning, TEXT("UpdatePartyLeader:  Unable to migrate player %s from reservation under leader %s to existing reservation with leader %s"),
-						*InPartyMemberId.ToString(), *PriorReservation->PartyLeader.ToString(), *NewPartyLeaderId.ToString());
+					UE_LOG(LogPartyBeacon, VeryVerbose, TEXT("UpdatePartyLeader:  Player %s already under party leader %s, no action taken"),
+						*InPartyMemberId.ToString(), *NewPartyLeaderId.ToString());
 				}
 			}
-			else
-			{
-				UE_LOG(LogPartyBeacon, VeryVerbose, TEXT("UpdatePartyLeader:  Player %s already under party leader %s, no action taken"),
-					*InPartyMemberId.ToString(), *NewPartyLeaderId.ToString());
-			}
 		}
 		else
 		{
