Don't remove player from beacon when he disconnects from a party in certain games.

diff --git a/Engine/Plugins/Online/OnlineFramework/Source/Party/Private/Party/SocialParty.cpp b/Engine/Plugins/Online/OnlineFramework/Source/Party/Private/Party/SocialParty.cpp
--- a/Engine/Plugins/Online/OnlineFramework/Source/Party/Private/Party/SocialParty.cpp
+++ b/Engine/Plugins/Online/OnlineFramework/Source/Party/Private/Party/SocialParty.cpp
@@ -1433,6 +1433,10 @@ void USocialParty::HandleReservationRequestComplete(EPartyReservationResult::Typ
 				// This player is already in our party. ApproveJIPRequest
 				PartyInterface->ApproveJIPRequest(*PendingApproval.RecipientId, GetPartyId(), *PendingApproval.SenderId, bReservationApproved, DenialReason);
 			}
+			else if (PendingApproval.bIsPlayerRemoval)
+			{
+				// We don't care about calling back the player when they are requesting a removal.
+			}
 			else
 			{
 				PartyInterface->ApproveJoinRequest(*PendingApproval.RecipientId, GetPartyId(), *PendingApproval.SenderId, bReservationApproved, DenialReason);
diff --git a/Engine/Plugins/Online/OnlineSubsystemUtils/Source/OnlineSubsystemUtils/Private/PartyBeaconHost.cpp b/Engine/Plugins/Online/OnlineSubsystemUtils/Source/OnlineSubsystemUtils/Private/PartyBeaconHost.cpp
--- a/Engine/Plugins/Online/OnlineSubsystemUtils/Source/OnlineSubsystemUtils/Private/PartyBeaconHost.cpp
+++ b/Engine/Plugins/Online/OnlineSubsystemUtils/Source/OnlineSubsystemUtils/Private/PartyBeaconHost.cpp
@@ -30,13 +30,13 @@ void APartyBeaconHost::PostInitProperties()
 #endif
 }
 
-bool APartyBeaconHost::InitHostBeacon(int32 InTeamCount, int32 InTeamSize, int32 InMaxReservations, FName InSessionName, int32 InForceTeamNum)
+bool APartyBeaconHost::InitHostBeacon(int32 InTeamCount, int32 InTeamSize, int32 InMaxReservations, FName InSessionName, int32 InForceTeamNum, bool bInEnableRemovalRequests)
 {
 	UE_LOG(LogPartyBeacon, Verbose, TEXT("InitHostBeacon TeamCount:%d TeamSize:%d MaxSize:%d"), InTeamCount, InTeamSize, InMaxReservations);
 	if (InMaxReservations > 0)
 	{
 		State = NewObject<UPartyBeaconState>(GetTransientPackage(), GetPartyBeaconHostClass());
-		if (State->InitState(InTeamCount, InTeamSize, InMaxReservations, InSessionName, InForceTeamNum))
+		if (State->InitState(InTeamCount, InTeamSize, InMaxReservations, InSessionName, InForceTeamNum, bInEnableRemovalRequests))
 		{
 			return true;
 		}
@@ -737,6 +737,12 @@ EPartyReservationResult::Type APartyBeaconHost::UpdatePartyReservation(const FPa
 		return EPartyReservationResult::ReservationDenied;
 	}
 
+	if (bIsRemovingMembers && !State->bEnableRemovalRequests)
+	{
+		// We are told not to accept removal requests.
+		return EPartyReservationResult::ReservationDenied;
+	}
+
 	if (ReservationUpdateRequest.IsValid() || bIsRemovingMembers)
 	{
 		if (bIsRemovingMembers)
diff --git a/Engine/Plugins/Online/OnlineSubsystemUtils/Source/OnlineSubsystemUtils/Private/PartyBeaconState.cpp b/Engine/Plugins/Online/OnlineSubsystemUtils/Source/OnlineSubsystemUtils/Private/PartyBeaconState.cpp
--- a/Engine/Plugins/Online/OnlineSubsystemUtils/Source/OnlineSubsystemUtils/Private/PartyBeaconState.cpp
+++ b/Engine/Plugins/Online/OnlineSubsystemUtils/Source/OnlineSubsystemUtils/Private/PartyBeaconState.cpp
@@ -123,7 +123,7 @@ UPartyBeaconState::UPartyBeaconState(const FObjectInitializer& ObjectInitializer
 {
 }
 
-bool UPartyBeaconState::InitState(int32 InTeamCount, int32 InTeamSize, int32 InMaxReservations, FName InSessionName, int32 InForceTeamNum)
+bool UPartyBeaconState::InitState(int32 InTeamCount, int32 InTeamSize, int32 InMaxReservations, FName InSessionName, int32 InForceTeamNum, bool bInEnableRemovalRequests)
 {
 	if (InMaxReservations > 0)
 	{
@@ -133,7 +133,7 @@ bool UPartyBeaconState::InitState(int32 InTeamCount, int32 InTeamSize, int32 InM
 		MaxReservations = InMaxReservations;
 		ForceTeamNum = InForceTeamNum;
 		Reservations.Empty(MaxReservations);
-
+		bEnableRemovalRequests = bInEnableRemovalRequests;
 		InitTeamArray();
 		return true;
 	}
diff --git a/Engine/Plugins/Online/OnlineSubsystemUtils/Source/OnlineSubsystemUtils/Public/PartyBeaconHost.h b/Engine/Plugins/Online/OnlineSubsystemUtils/Source/OnlineSubsystemUtils/Public/PartyBeaconHost.h
--- a/Engine/Plugins/Online/OnlineSubsystemUtils/Source/OnlineSubsystemUtils/Public/PartyBeaconHost.h
+++ b/Engine/Plugins/Online/OnlineSubsystemUtils/Source/OnlineSubsystemUtils/Public/PartyBeaconHost.h
@@ -73,10 +73,11 @@ class ONLINESUBSYSTEMUTILS_API APartyBeaconHost : public AOnlineBeaconHostObject
 	 * @param InMaxReservation max number of reservations allowed
 	 * @param InSessionName name of session related to the beacon
 	 * @param InForceTeamNum team to force players on if applicable (usually only 1 team games)
+	 * @param bInEnableRemovalRequests allow clients to remove players from beacon.
 	 *
 	 * @return true if successful created, false otherwise
 	 */
-	virtual bool InitHostBeacon(int32 InTeamCount, int32 InTeamSize, int32 InMaxReservations, FName InSessionName, int32 InForceTeamNum = 0);
+	virtual bool InitHostBeacon(int32 InTeamCount, int32 InTeamSize, int32 InMaxReservations, FName InSessionName, int32 InForceTeamNum = 0, bool bInEnableRemovalRequests = true);
 
 	/**
 	 * Initialize the party host beacon from a previous state/configuration
diff --git a/Engine/Plugins/Online/OnlineSubsystemUtils/Source/OnlineSubsystemUtils/Public/PartyBeaconState.h b/Engine/Plugins/Online/OnlineSubsystemUtils/Source/OnlineSubsystemUtils/Public/PartyBeaconState.h
--- a/Engine/Plugins/Online/OnlineSubsystemUtils/Source/OnlineSubsystemUtils/Public/PartyBeaconState.h
+++ b/Engine/Plugins/Online/OnlineSubsystemUtils/Source/OnlineSubsystemUtils/Public/PartyBeaconState.h
@@ -237,7 +237,7 @@ class ONLINESUBSYSTEMUTILS_API UPartyBeaconState : public UObject
 	 *
 	 * @return true if successful created, false otherwise
 	 */
-	virtual bool InitState(int32 InTeamCount, int32 InTeamSize, int32 InMaxReservations, FName InSessionName, int32 InForceTeamNum);
+	virtual bool InitState(int32 InTeamCount, int32 InTeamSize, int32 InMaxReservations, FName InSessionName, int32 InForceTeamNum, bool bEnableRemovalRequests);
 
 	/**
 	 * Reconfigures the beacon for a different team/player count configuration
@@ -558,6 +558,9 @@ protected:
 	/** Are multiple consoles types allowed to play together */
 	UPROPERTY(Config)
 	bool bRestrictCrossConsole;
+	/** Process requests from clients to remove players from beacon */
+	UPROPERTY(Transient)
+	bool bEnableRemovalRequests;
 
 	/** Current reservations in the system */
 	UPROPERTY(Transient)
