Fix the address mapping issue with IPv6 locked devices when user input is not pushed through a sanitizer and leads to a protocol conflict. This change introduces address synthesis on protocol mismatches as a patch.

A more robust solution will appear later.

diff --git a/Engine/Plugins/Online/OnlineSubsystemUtils/Source/OnlineSubsystemUtils/Private/IpConnection.cpp b/Engine/Plugins/Online/OnlineSubsystemUtils/Source/OnlineSubsystemUtils/Private/IpConnection.cpp
--- a/Engine/Plugins/Online/OnlineSubsystemUtils/Source/OnlineSubsystemUtils/Private/IpConnection.cpp
+++ b/Engine/Plugins/Online/OnlineSubsystemUtils/Source/OnlineSubsystemUtils/Private/IpConnection.cpp
@@ -29,6 +29,7 @@ Notes:
 DECLARE_CYCLE_STAT(TEXT("IpConnection InitRemoteConnection"), Stat_IpConnectionInitRemoteConnection, STATGROUP_Net);
 DECLARE_CYCLE_STAT(TEXT("IpConnection Socket SendTo"), STAT_IpConnection_SendToSocket, STATGROUP_Net);
 DECLARE_CYCLE_STAT(TEXT("IpConnection WaitForSendTasks"), STAT_IpConnection_WaitForSendTasks, STATGROUP_Net);
+DECLARE_CYCLE_STAT(TEXT("IpConnection Address Synthesis"), STAT_IpConnection_AddressSynthesis, STATGROUP_Net);
 
 TAutoConsoleVariable<int32> CVarNetIpConnectionUseSendTasks(
 	TEXT("net.IpConnectionUseSendTasks"),
@@ -69,6 +70,27 @@ void UIpConnection::InitLocalConnection(UNetDriver* InDriver, class FSocket* InS
 	// Get numerical address directly.
 	RemoteAddr = InDriver->GetSocketSubsystem()->CreateInternetAddr();
 	RemoteAddr->SetIp(*InURL.Host, bIsValid);
+
+	// If the protocols do not match, attempt to synthesize the address so they do.
+	if (bIsValid && InSocket->GetProtocol() != RemoteAddr->GetProtocolType())
+	{
+		SCOPE_CYCLE_COUNTER(STAT_IpConnection_AddressSynthesis);
+
+		// We want to use GAI to create the address with the correct protocol.
+		const FAddressInfoResult MapRequest = InDriver->GetSocketSubsystem()->GetAddressInfo(*InURL.Host, nullptr, 
+			EAddressInfoFlags::AllResultsWithMapping | EAddressInfoFlags::OnlyUsableAddresses, InSocket->GetProtocol());
+
+		// Set the remote addr provided we have information.
+		if (MapRequest.ReturnCode == SE_NO_ERROR && MapRequest.Results.Num() > 0)
+		{
+			RemoteAddr = MapRequest.Results[0].Address->Clone();
+		}
+		else
+		{
+			UE_LOG(LogNet, Warning, TEXT("IpConnection::InitConnection: Address protocols do not match and cannot be synthesized to a similar address, this will likely lead to issues!"));
+		}
+	}
+
 	RemoteAddr->SetPort(InURL.Port);
 
 	// Try to resolve it if it failed
