Added new FReceivedPacketView and FPacketBufferView types, for tidying up function parameter lists, and refactoring the net packet pipeline with more strict separation between non-modifiable 'views' into packet data, and modifiable but fixed-allocation 'views' into packet buffers.

diff --git a/Engine/Plugins/Online/OnlineSubsystemUtils/Source/OnlineSubsystemUtils/Classes/IpNetDriver.h b/Engine/Plugins/Online/OnlineSubsystemUtils/Source/OnlineSubsystemUtils/Classes/IpNetDriver.h
--- a/Engine/Plugins/Online/OnlineSubsystemUtils/Source/OnlineSubsystemUtils/Classes/IpNetDriver.h
+++ b/Engine/Plugins/Online/OnlineSubsystemUtils/Source/OnlineSubsystemUtils/Classes/IpNetDriver.h
@@ -25,8 +25,9 @@ class ONLINESUBSYSTEMUTILS_API UIpNetDriver : public UNetDriver
 {
 	friend class FPacketIterator;
 
-    GENERATED_UCLASS_BODY()
+    GENERATED_BODY()
 
+public:
 	/** Should port unreachable messages be logged */
     UPROPERTY(Config)
     uint32 LogPortUnreach:1;
@@ -45,6 +46,10 @@ class ONLINESUBSYSTEMUTILS_API UIpNetDriver : public UNetDriver
 	/** If pausing socket receives, the time at which this should end */
 	float PauseReceiveEnd;
 
+
+	/** Base constructor */
+	UIpNetDriver(const FObjectInitializer& ObjectInitializer);
+
 	//~ Begin UNetDriver Interface.
 	virtual bool IsAvailable() const override;
 	virtual bool InitBase(bool bInitAsClient, FNetworkNotify* InNotify, const FURL& URL, bool bReuseAddressAndPort, FString& Error) override;
@@ -69,8 +74,33 @@ class ONLINESUBSYSTEMUTILS_API UIpNetDriver : public UNetDriver
 	 * @param CountBytesRef		The packet size (may be modified)
 	 * @return					If a new NetConnection is created, returns the net connection
 	 */
-	UNetConnection* ProcessConnectionlessPacket(const TSharedRef<FInternetAddr>& Address, uint8* Data, int32& CountBytesRef);
+	UE_DEPRECATED(4.24, "ProcessConnectionlessPacket has a new API, and will become private soon.")
+	UNetConnection* ProcessConnectionlessPacket(const TSharedRef<FInternetAddr>& Address, uint8* Data, int32& CountBytesRef)
+	{
+		FReceivedPacketView PacketView;
+
+		PacketView.Data = MakeArrayView(Data, CountBytesRef);
+		PacketView.Address = Address;
+		PacketView.Error = SE_NO_ERROR;
+
+		UNetConnection* ReturnVal = ProcessConnectionlessPacket(PacketView, { Data, MAX_PACKET_SIZE } );
+
+		CountBytesRef = PacketView.Data.Num();
+
+		return ReturnVal;
+	}
+
+private:
+	/**
+	 * Process packets not associated with a NetConnection, performing handshaking and NetConnection creation or remapping as necessary.
+	 *
+	 * @param PacketRef			A view of the received packet (may output a new packet view, possibly using WorkingBuffer)
+	 * @param WorkingBuffer		A buffer for storing processed packet data (may be the buffer the input ReceivedPacketRef points to)
+	 * @return					If a new NetConnection is created, returns the net connection
+	 */
+	UNetConnection* ProcessConnectionlessPacket(FReceivedPacketView& PacketRef, const FPacketBufferView& WorkingBuffer);
 
+public:
 	//~ Begin UIpNetDriver Interface.
 	virtual FSocket * CreateSocket();
 
diff --git a/Engine/Plugins/Online/OnlineSubsystemUtils/Source/OnlineSubsystemUtils/Private/IpNetDriver.cpp b/Engine/Plugins/Online/OnlineSubsystemUtils/Source/OnlineSubsystemUtils/Private/IpNetDriver.cpp
--- a/Engine/Plugins/Online/OnlineSubsystemUtils/Source/OnlineSubsystemUtils/Private/IpNetDriver.cpp
+++ b/Engine/Plugins/Online/OnlineSubsystemUtils/Source/OnlineSubsystemUtils/Private/IpNetDriver.cpp
@@ -90,7 +90,7 @@ class FPacketIterator
 private:
 	struct FCachedPacket
 	{
-		/** Whether or not socket receive succeeded. Don't rely on the Error field for this, due to implementation/platform uncertainties. */
+		/** Whether socket receive succeeded. Don't rely on the Error field for this, due to implementation/platform uncertainties. */
 		bool bRecvSuccess;
 
 		/** Pre-allocated Data field, for storing packets of any expected size */
@@ -140,24 +140,28 @@ private:
 	/**
 	 * Retrieves the packet information from the current iteration. Avoid calling more than once, per iteration.
 	 *
-	 * @param OutData		Returns a pointer to the packet data
-	 * @param OutBytesRead	Returns the number of bytes read
-	 * @param OutError		Returns the socket error, if receiving failed
-	 * @param OutAddr		Returns the address the packet was received from
+	 * @param OutPacket		Outputs a view to the received packet data
 	 * @return				Returns whether or not receiving was successful for the current packet
 	 */
-	bool GetCurrentPacket(const uint8*& OutData, uint32& OutBytesRead, ESocketErrors& OutError, TSharedPtr<FInternetAddr>& OutAddr)
+	bool GetCurrentPacket(FReceivedPacketView& OutPacket)
 	{
 		bool bRecvSuccess = false;
 
-		OutData = CurrentPacket.Data.GetData();
-		OutBytesRead = CurrentPacket.Data.Num();
-		OutError = CurrentPacket.Error;
-		OutAddr = CurrentPacket.Address;
+		OutPacket.Data = MakeArrayView(CurrentPacket.Data);
+		OutPacket.Error = CurrentPacket.Error;
+		OutPacket.Address = CurrentPacket.Address;
 		bRecvSuccess = CurrentPacket.bRecvSuccess;
 
 		return bRecvSuccess;
 	}
+	
+	/**
+	 * Returns a view of the iterator's packet buffer, for updating packet data as it's processed, and generating new packet views
+	 */
+	FPacketBufferView GetWorkingBuffer()
+	{
+		return { CurrentPacket.Data.GetData(), MAX_PACKET_SIZE };
+	}
 
 	/**
 	 * Advances the current packet to the next iteration
@@ -578,38 +582,34 @@ void UIpNetDriver::TickDispatch(float DeltaTime)
 			}
 		}
 
-		uint8* Data = nullptr;
-		int32 BytesRead = 0;
-		ESocketErrors Error = SE_NO_ERROR;
-		TSharedPtr<FInternetAddr> FromAddrPtr;
-		bool bOk = It.GetCurrentPacket((const uint8*&)Data, (uint32&)BytesRead, Error, FromAddrPtr);
-		TSharedRef<FInternetAddr> FromAddr = FromAddrPtr.ToSharedRef();
+		FReceivedPacketView ReceivedPacket;
+		bool bOk = It.GetCurrentPacket(ReceivedPacket);
+		const TSharedRef<const FInternetAddr> FromAddr = ReceivedPacket.Address.ToSharedRef();
 		UNetConnection* Connection = nullptr;
 		UIpConnection* const MyServerConnection = GetServerConnection();
 
 		if (bOk)
 		{
 			// Immediately stop processing (continuing to next receive), for empty packets (usually a DDoS)
-			if (BytesRead == 0)
+			if (ReceivedPacket.Data.Num() == 0)
 			{
 				DDoS.IncBadPacketCounter();
 				continue;
 			}
 
-			FPacketAudit::NotifyLowLevelReceive(Data, BytesRead);
+			FPacketAudit::NotifyLowLevelReceive((uint8*)ReceivedPacket.Data.GetData(), ReceivedPacket.Data.Num());
 		}
 		else
 		{
-			if (IsRecvFailBlocking(Error))
+			if (IsRecvFailBlocking(ReceivedPacket.Error))
 			{
-				// No data or no error? (SE_ECONNABORTED is for PS4 LAN cable pulls)
 				break;
 			}
-			else if (Error != SE_ECONNRESET && Error != SE_UDP_ERR_PORT_UNREACH)
+			else if (ReceivedPacket.Error != SE_ECONNRESET && ReceivedPacket.Error != SE_UDP_ERR_PORT_UNREACH)
 			{
 				// MalformedPacket: Client tried receiving a packet that exceeded the maximum packet limit
 				// enforced by the server
-				if (Error == SE_EMSGSIZE)
+				if (ReceivedPacket.Error == SE_EMSGSIZE)
 				{
 					DDoS.IncBadPacketCounter();
 
@@ -639,8 +639,8 @@ void UIpNetDriver::TickDispatch(float DeltaTime)
 				}
 
 				FString ErrorString = FString::Printf(TEXT("UIpNetDriver::TickDispatch: Socket->RecvFrom: %i (%s) from %s"),
-					static_cast<int32>(Error),
-					SocketSubsystem->GetSocketError(Error),
+					static_cast<int32>(ReceivedPacket.Error),
+					SocketSubsystem->GetSocketError(ReceivedPacket.Error),
 					*FromAddr->ToString(true));
 
 
@@ -688,8 +688,7 @@ void UIpNetDriver::TickDispatch(float DeltaTime)
 
 		if (Connection == nullptr)
 		{
-			const TSharedRef<const FInternetAddr> ConstFromAddr = FromAddr;
-			UNetConnection** Result = MappedClientConnections.Find(ConstFromAddr);
+			UNetConnection** Result = MappedClientConnections.Find(FromAddr);
 
 			if (Result != nullptr)
 			{
@@ -773,11 +772,13 @@ void UIpNetDriver::TickDispatch(float DeltaTime)
 
 				if (bAcceptingConnection)
 				{
-					UE_CLOG(!DDoS.CheckLogRestrictions(), LogNet, Log, TEXT("NotifyAcceptingConnection accepted from: %s"), *FromAddr->ToString(true));
+					UE_CLOG(!DDoS.CheckLogRestrictions(), LogNet, Log, TEXT("NotifyAcceptingConnection accepted from: %s"),
+								*FromAddr->ToString(true));
 
-					Connection = ProcessConnectionlessPacket(FromAddr, Data, BytesRead);
+					FPacketBufferView WorkingBuffer = It.GetWorkingBuffer();
 
-					bIgnorePacket = BytesRead == 0;
+					Connection = ProcessConnectionlessPacket(ReceivedPacket, WorkingBuffer);
+					bIgnorePacket = ReceivedPacket.Data.Num() == 0;
 				}
 				else
 				{
@@ -794,7 +795,7 @@ void UIpNetDriver::TickDispatch(float DeltaTime)
 					DDoS.CondCheckNetConnLimits();
 				}
 
-				Connection->ReceivedRawPacket(Data, BytesRead);
+				Connection->ReceivedRawPacket((uint8*)ReceivedPacket.Data.GetData(), ReceivedPacket.Data.Num());
 			}
 		}
 	}
@@ -809,10 +810,11 @@ void UIpNetDriver::TickDispatch(float DeltaTime)
 	}
 }
 
-UNetConnection* UIpNetDriver::ProcessConnectionlessPacket(const TSharedRef<FInternetAddr>& Address, uint8* Data, int32& CountBytesRef)
+UNetConnection* UIpNetDriver::ProcessConnectionlessPacket(FReceivedPacketView& PacketRef, const FPacketBufferView& WorkingBuffer)
 {
 	UNetConnection* ReturnVal = nullptr;
 	TSharedPtr<StatelessConnectHandlerComponent> StatelessConnect;
+	const TSharedPtr<FInternetAddr>& Address = PacketRef.Address;
 	FString IncomingAddress = Address->ToString(true);
 	bool bPassedChallenge = false;
 	bool bRestartedHandshake = false;
@@ -822,9 +824,10 @@ UNetConnection* UIpNetDriver::ProcessConnectionlessPacket(const TSharedRef<FInte
 	{
 		StatelessConnect = StatelessConnectComponent.Pin();
 
-		const ProcessedPacket UnProcessedPacket = ConnectionlessHandler->IncomingConnectionless(Address, Data, CountBytesRef);
+		const ProcessedPacket HandlerResult = ConnectionlessHandler->IncomingConnectionless(Address,
+																		(uint8*)PacketRef.Data.GetData(), PacketRef.Data.Num());
 
-		if (!UnProcessedPacket.bError)
+		if (!HandlerResult.bError)
 		{
 			bPassedChallenge = StatelessConnect->HasPassedChallenge(Address, bRestartedHandshake);
 
@@ -892,14 +895,17 @@ UNetConnection* UIpNetDriver::ProcessConnectionlessPacket(const TSharedRef<FInte
 					}
 				}
 
-				CountBytesRef = FMath::DivideAndRoundUp(UnProcessedPacket.CountBits, 8);
 
-				if (CountBytesRef > 0)
+				int32 NewCountBytes = FMath::DivideAndRoundUp(HandlerResult.CountBits, 8);
+
+				if (NewCountBytes > 0)
 				{
-					FMemory::Memcpy(Data, UnProcessedPacket.Data, CountBytesRef);
+					FMemory::Memcpy(WorkingBuffer.Buffer.GetData(), HandlerResult.Data, NewCountBytes);
 
 					bIgnorePacket = false;
 				}
+
+				PacketRef.Data = MakeArrayView(WorkingBuffer.Buffer.GetData(), NewCountBytes);
 			}
 		}
 	}
@@ -960,12 +966,12 @@ UNetConnection* UIpNetDriver::ProcessConnectionlessPacket(const TSharedRef<FInte
 	}
 	else
 	{
-		UE_LOG(LogNet, VeryVerbose, TEXT( "Server failed post-challenge connection from: %s" ), *IncomingAddress);
+		UE_LOG(LogNet, VeryVerbose, TEXT("Server failed post-challenge connection from: %s"), *IncomingAddress);
 	}
 
 	if (bIgnorePacket)
 	{
-		CountBytesRef = 0;
+		PacketRef.Data = MakeArrayView(PacketRef.Data.GetData(), 0);
 	}
 
 	return ReturnVal;
diff --git a/Engine/Source/Runtime/Sockets/Public/SocketTypes.h b/Engine/Source/Runtime/Sockets/Public/SocketTypes.h
--- a/Engine/Source/Runtime/Sockets/Public/SocketTypes.h
+++ b/Engine/Source/Runtime/Sockets/Public/SocketTypes.h
@@ -4,6 +4,9 @@
 
 #include "CoreMinimal.h"
 
+// Forward declarations
+class FInternetAddr;
+
 /** Indicates the socket protocol of socket being used, typically for BSD Sockets */
 enum class ESocketProtocolFamily : uint8
 {
@@ -176,4 +179,37 @@ enum class ESocketShutdownMode
 	 * Disables reading and writing on the socket.
 	 */
 	ReadWrite
-};
\ No newline at end of file
+};
+
+
+/**
+ * Represents a view of a buffer for storing packets. Buffer contents may be modified, but the allocation can not be resized.
+ * Should only be stored as a local variable within functions that handle received packets.
+ */
+struct FPacketBufferView
+{
+	/** View of the packet buffer, with Num() representing allocated size. Internal buffer data can be modified, but not view/size. */
+	const TArrayView<uint8>		Buffer;
+
+
+	FPacketBufferView(uint8* InData, int32 MaxBufferSize)
+		: Buffer(MakeArrayView<uint8>(InData, MaxBufferSize))
+	{
+	}
+};
+
+/**
+ * Represents a view of a received packet, which may be modified to update Data it points to and Data size, as a packet is processed.
+ * Should only be stored as a local variable within functions that handle received packets.
+ */
+struct FReceivedPacketView
+{
+	/** View of packet data, with Num() representing BytesRead - can reassign to point elsewhere, but don't use to modify packet data */
+	TArrayView<const uint8>		Data;
+
+	/** Receive address for the packet */
+	TSharedPtr<FInternetAddr>	Address;
+
+	/** Error if receiving a packet failed */
+	ESocketErrors				Error;
+};
