Implemented IpNetDriver use of Linux-specific RecvMulti, and tied the packet receive timestamp into the ping calculation code (affects serverside ping calculations only - CVar net.PingUsePacketRecvTime - factoring this into client calculations requires a big redesign of how pings are done).

diff --git a/Engine/Plugins/Online/OnlineSubsystemUtils/Source/OnlineSubsystemUtils/Classes/IpNetDriver.h b/Engine/Plugins/Online/OnlineSubsystemUtils/Source/OnlineSubsystemUtils/Classes/IpNetDriver.h
--- a/Engine/Plugins/Online/OnlineSubsystemUtils/Source/OnlineSubsystemUtils/Classes/IpNetDriver.h
+++ b/Engine/Plugins/Online/OnlineSubsystemUtils/Source/OnlineSubsystemUtils/Classes/IpNetDriver.h
@@ -19,6 +19,7 @@ class Error;
 class FInternetAddr;
 class FNetworkNotify;
 class FSocket;
+struct FRecvMulti;
 
 UCLASS(transient, config=Engine)
 class ONLINESUBSYSTEMUTILS_API UIpNetDriver : public UNetDriver
@@ -225,4 +226,7 @@ private:
 
 	/** Receive thread object. */
 	TUniquePtr<FRunnableThread> SocketReceiveThread;
+
+	/** The preallocated state/buffers, for efficiently executing RecvMulti */
+	TUniquePtr<FRecvMulti> RecvMultiState;
 };
diff --git a/Engine/Plugins/Online/OnlineSubsystemUtils/Source/OnlineSubsystemUtils/Private/IpNetDriver.cpp b/Engine/Plugins/Online/OnlineSubsystemUtils/Source/OnlineSubsystemUtils/Private/IpNetDriver.cpp
--- a/Engine/Plugins/Online/OnlineSubsystemUtils/Source/OnlineSubsystemUtils/Private/IpNetDriver.cpp
+++ b/Engine/Plugins/Online/OnlineSubsystemUtils/Source/OnlineSubsystemUtils/Private/IpNetDriver.cpp
@@ -24,6 +24,7 @@ Notes:
 
 #include "IPAddress.h"
 #include "Sockets.h"
+#include "Serialization/ArchiveCountMem.h"
 
 /** For backwards compatibility with the engine stateless connect code */
 #ifndef STATELESSCONNECT_HAS_RANDOM_SEQUENCE
@@ -74,6 +75,24 @@ TAutoConsoleVariable<int32> CVarNetIpNetDriverReceiveThreadPollTimeMS(
 	250,
 	TEXT("If net.IpNetDriverUseReceiveThread is true, the number of milliseconds to use as the timeout value for FSocket::Wait on the receive thread. A negative value means to wait indefinitely (FSocket::Shutdown should cancel it though)."));
 
+TAutoConsoleVariable<int32> CVarNetUseRecvMulti(
+	TEXT("net.UseRecvMulti"),
+	0,
+	TEXT("If true, and if running on a Unix/Linux platform, multiple packets will be retrieved from the socket with one syscall, ")
+		TEXT("improving performance and also allowing retrieval of timestamp information."));
+
+TAutoConsoleVariable<int32> CVarRecvMultiCapacity(
+	TEXT("net.RecvMultiCapacity"),
+	2048,
+	TEXT("When RecvMulti is enabled, this is the number of packets it is allocated to handle per call - ")
+		TEXT("bigger is better (especially under a DDoS), but keep an eye on memory cost."));
+
+TAutoConsoleVariable<int32> CVarNetUseRecvTimestamps(
+	TEXT("net.UseRecvTimestamps"),
+	0,
+	TEXT("If true and if net.UseRecvMulti is also true, on a Unix/Linux platform, ")
+		TEXT("the kernel timestamp will be retrieved for each packet received, providing more accurate ping calculations."));
+
 
 
 /**
@@ -115,8 +134,13 @@ private:
 		, Socket(Driver->Socket)
 		, SocketReceiveThreadRunnable(Driver->SocketReceiveThreadRunnable.Get())
 		, CurrentPacket()
+		, RecvMultiIdx(0)
+		, RecvMultiPacketCount(0)
 	{
-		if (SocketSubsystem != nullptr)
+		RMState = Driver->RecvMultiState.Get();
+		bUseRecvMulti = CVarNetUseRecvMulti.GetValueOnAnyThread() != 0 && RMState != nullptr;
+
+		if (!bUseRecvMulti && SocketSubsystem != nullptr)
 		{
 			CurrentPacket.Address = SocketSubsystem->CreateInternetAddr();
 		}
@@ -147,14 +171,52 @@ private:
 	{
 		bool bRecvSuccess = false;
 
-		OutPacket.Data = MakeArrayView(CurrentPacket.Data);
-		OutPacket.Error = CurrentPacket.Error;
-		OutPacket.Address = CurrentPacket.Address;
-		bRecvSuccess = CurrentPacket.bRecvSuccess;
+		if (bUseRecvMulti)
+		{
+			RMState->GetPacket(RecvMultiIdx, OutPacket);
+			bRecvSuccess = true;
+		}
+		else
+		{
+			OutPacket.Data = MakeArrayView(CurrentPacket.Data);
+			OutPacket.Error = CurrentPacket.Error;
+			OutPacket.Address = CurrentPacket.Address;
+			bRecvSuccess = CurrentPacket.bRecvSuccess;
+		}
 
 		return bRecvSuccess;
 	}
-	
+
+	/**
+	 * Retrieves the packet timestamp information from the current iteration. As above, avoid calling more than once.
+	 *
+	 * @param ForConnection		The connection we are retrieving timestamp information for
+	 */
+	void GetCurrentPacketTimestamp(UNetConnection* ForConnection)
+	{
+		FPacketTimestamp CurrentTimestamp;
+		bool bIsLocalTimestamp = false;
+		bool bSuccess = false;
+
+		if (bUseRecvMulti)
+		{
+			RMState->GetPacketTimestamp(RecvMultiIdx, CurrentTimestamp);
+			bIsLocalTimestamp = false;
+			bSuccess = true;
+		}
+		else if (CurrentPacket.PacketTimestamp != 0.0)
+		{
+			CurrentTimestamp.Timestamp = FTimespan::FromSeconds(CurrentPacket.PacketTimestamp);
+			bIsLocalTimestamp = true;
+			bSuccess = true;
+		}
+
+		if (bSuccess)
+		{
+			ForConnection->SetPacketOSReceiveTime(CurrentTimestamp, bIsLocalTimestamp);
+		}
+	}
+
 	/**
 	 * Returns a view of the iterator's packet buffer, for updating packet data as it's processed, and generating new packet views
 	 */
@@ -168,7 +230,23 @@ private:
 	 */
 	void AdvanceCurrentPacket()
 	{
-		bBreak = !ReceiveSinglePacket();
+		if (bUseRecvMulti)
+		{
+			if (RecvMultiPacketCount == 0 || ((RecvMultiIdx + 1) >= RecvMultiPacketCount))
+			{
+				AdvanceRecvMultiState();
+			}
+			else
+			{
+				RecvMultiIdx++;
+			}
+
+			// At this point, bBreak will be set, or RecvMultiPacketCount will be > 0
+		}
+		else
+		{
+			bBreak = !ReceiveSinglePacket();
+		}
 	}
 
 	/**
@@ -305,6 +383,63 @@ private:
 		return bReceivedPacketOrError;
 	}
 
+	/**
+	 * Load a fresh batch of RecvMulti packets
+	 */
+	void AdvanceRecvMultiState()
+	{
+		RecvMultiIdx = 0;
+		RecvMultiPacketCount = 0;
+
+		while (true)
+		{
+			bool bRecvMultiOk = Socket != nullptr ? Socket->RecvMulti(*RMState) : false;
+
+			if (!bRecvMultiOk)
+			{
+				ESocketErrors RecvMultiError = (SocketSubsystem != nullptr ? SocketSubsystem->GetLastErrorCode() : SE_NO_ERROR);
+
+				if (UIpNetDriver::IsRecvFailBlocking(RecvMultiError))
+				{
+					bBreak = true;
+					break;
+				}
+				else
+				{
+					// When the Linux recvmmsg syscall encounters an error after successfully receiving at least one packet,
+					// it won't return an error until called again, but this error can be overwritten before recvmmsg is called again.
+					// This makes the error handling for recvmmsg unreliable. Continue until the socket blocks.
+
+					// Continue is safe, as 0 packets have been received
+					continue;
+				}
+			}
+
+			// Extreme-early-out. NetConnection per frame time limit, limits all packet processing - RecvMulti drops all packets at once
+			if (DDoS.ShouldBlockNetConnPackets())
+			{
+				int32 NumDropped = RMState->GetNumPackets();
+
+				DDoS.IncDroppedPacketCounter(NumDropped);
+
+				// Have a threshold, to stop the RecvMulti syscall spinning with low packet counts - let the socket buffer build up
+				if (NumDropped > 10)
+				{
+					continue;
+				}
+				else
+				{
+					bBreak = true;
+					break;
+				}
+			}
+
+			RecvMultiPacketCount = RMState->GetNumPackets();
+
+			break;
+		}
+	}
+
 
 private:
 	/** Specified internally, when the packet iterator should break/stop (no packets, DDoS limits triggered, etc.) */
@@ -325,6 +460,18 @@ private:
 
 	/** Stores information for the current packet being received (when using single-receive mode) */
 	FCachedPacket CurrentPacket;
+
+	/** Stores information for receiving packets using RecvMulti */
+	FRecvMulti* RMState;
+
+	/** Whether or not RecvMulti is enabled/supported */
+	bool bUseRecvMulti;
+
+	/** The RecvMulti index of the next packet to be received (if RecvMultiPacketCount > 0) */
+	int32 RecvMultiIdx;
+
+	/** The number of packets waiting to be read from the FRecvMulti state */
+	int32 RecvMultiPacketCount;
 };
 
 /**
@@ -338,6 +485,7 @@ UIpNetDriver::UIpNetDriver(const FObjectInitializer& ObjectInitializer)
 	, ServerDesiredSocketSendBufferBytes(0x20000)
 	, ClientDesiredSocketReceiveBufferBytes(0x8000)
 	, ClientDesiredSocketSendBufferBytes(0x8000)
+	, RecvMultiState(nullptr)
 {
 }
 
@@ -463,6 +611,47 @@ bool UIpNetDriver::InitBase( bool bInitAsClient, FNetworkNotify* InNotify, const
 		SocketReceiveThread.Reset(FRunnableThread::Create(SocketReceiveThreadRunnable.Get(), *FString::Printf(TEXT("IpNetDriver Receive Thread"), *NetDriverName.ToString())));
 	}
 
+
+	bool bRecvMultiEnabled = CVarNetUseRecvMulti.GetValueOnAnyThread() != 0;
+	bool bRecvThreadEnabled = CVarNetIpNetDriverUseReceiveThread.GetValueOnAnyThread() != 0;
+
+	if (bRecvMultiEnabled && !bRecvThreadEnabled)
+	{
+		bool bSupportsRecvMulti = SocketSubsystem->IsSocketRecvMultiSupported();
+
+		if (bSupportsRecvMulti)
+		{
+			bool bRetrieveTimestamps = CVarNetUseRecvTimestamps.GetValueOnAnyThread() != 0;
+
+			if (bRetrieveTimestamps)
+			{
+				Socket->SetRetrieveTimestamp(true);
+			}
+
+
+			ERecvMultiFlags RecvMultiFlags = bRetrieveTimestamps ? ERecvMultiFlags::RetrieveTimestamps : ERecvMultiFlags::None;
+			int32 MaxRecvMultiPackets = FMath::Max(32, CVarRecvMultiCapacity.GetValueOnAnyThread());
+
+			RecvMultiState = SocketSubsystem->CreateRecvMulti(MaxRecvMultiPackets, MAX_PACKET_SIZE, RecvMultiFlags);
+
+
+			FArchiveCountMem MemArc(nullptr);
+
+			RecvMultiState->CountBytes(MemArc);
+
+			UE_LOG(LogNet, Log, TEXT("NetDriver RecvMulti state size: %i, Retrieve Timestamps: %i"), MemArc.GetMax(),
+					(uint32)bRetrieveTimestamps);
+		}
+		else
+		{
+			UE_LOG(LogNet, Warning, TEXT("NetDriver could not enable RecvMulti, as current socket subsystem does not support it."));
+		}
+	}
+	else if (bRecvThreadEnabled)
+	{
+		UE_LOG(LogNet, Warning, TEXT("NetDriver RecvMulti is not yet supported with the Receive Thread enabled."));
+	}
+
 	// Success.
 	return true;
 }
@@ -544,6 +733,7 @@ void UIpNetDriver::TickDispatch(float DeltaTime)
 	const double StartReceiveTime = FPlatformTime::Seconds();
 	double AlarmTime = StartReceiveTime + GIpNetDriverMaxDesiredTimeSliceBeforeAlarmSecs;
 	const bool bSlowFrameChecks = OnNetworkProcessingCausingSlowFrame.IsBound();
+	bool bRetrieveTimestamps = CVarNetUseRecvTimestamps.GetValueOnAnyThread() != 0;
 
 	const bool bCheckReceiveTime = (MaxSecondsInReceive > 0.0) && (NbPacketsBetweenReceiveTimeTest > 0);
 	const double BailOutTime = StartReceiveTime + MaxSecondsInReceive;
@@ -795,6 +985,11 @@ void UIpNetDriver::TickDispatch(float DeltaTime)
 					DDoS.CondCheckNetConnLimits();
 				}
 
+				if (bRetrieveTimestamps)
+				{
+					It.GetCurrentPacketTimestamp(Connection);
+				}
+
 				Connection->ReceivedRawPacket((uint8*)ReceivedPacket.Data.GetData(), ReceivedPacket.Data.Num());
 			}
 		}
diff --git a/Engine/Source/Runtime/Engine/Classes/Engine/NetConnection.h b/Engine/Source/Runtime/Engine/Classes/Engine/NetConnection.h
--- a/Engine/Source/Runtime/Engine/Classes/Engine/NetConnection.h
+++ b/Engine/Source/Runtime/Engine/Classes/Engine/NetConnection.h
@@ -392,6 +392,11 @@ public:
 	/** Time when connection request was first initiated */
 	float			ConnectTime;
 
+private:
+	FPacketTimestamp	LastOSReceiveTime;		// Last time a packet was received at the OS/NIC layer
+	bool				bIsOSReceiveTimeLocal;	// Whether LastOSReceiveTime uses the same clock as the game, or needs translating
+
+public:
 	// Merge info.
 	FBitWriterMark  LastStart;				// Most recently sent bunch start.
 	FBitWriterMark  LastEnd;				// Most recently sent bunch end.
@@ -1143,6 +1148,18 @@ public:
 	 */
 	void FlushPacketOrderCache(bool bFlushWholeCache=false);
 
+	/**
+	 * Sets the OS/NIC level timestamp, for the last packet that was received
+	 *
+	 * @param InOSReceiveTime			The OS/NIC level timestamp, for the last packet that was received
+	 * @param bInIsOSReceiveTimeLocal	Whether the input timestamp is based on the same clock as the game thread, or needs translating
+	 */
+	void SetPacketOSReceiveTime(const FPacketTimestamp& InOSReceiveTime, bool bInIsOSReceiveTimeLocal)
+	{
+		LastOSReceiveTime = InOSReceiveTime;
+		bIsOSReceiveTimeLocal = bInIsOSReceiveTimeLocal;
+	}
+
 	/** Called when an actor channel is open and knows its NetGUID. */
 	ENGINE_API virtual void NotifyActorNetGUID(UActorChannel* Channel) {}
 
diff --git a/Engine/Source/Runtime/Engine/Private/NetConnection.cpp b/Engine/Source/Runtime/Engine/Private/NetConnection.cpp
--- a/Engine/Source/Runtime/Engine/Private/NetConnection.cpp
+++ b/Engine/Source/Runtime/Engine/Private/NetConnection.cpp
@@ -38,6 +38,8 @@
 
 static TAutoConsoleVariable<int32> CVarPingExcludeFrameTime( TEXT( "net.PingExcludeFrameTime" ), 0, TEXT( "Calculate RTT time between NIC's of server and client." ) );
 
+static TAutoConsoleVariable<int32> CVarPingUsePacketRecvTime(TEXT("net.PingUsePacketRecvTime"), 0, TEXT("Use OS or Receive Thread packet receive time, for calculating the ping. Excludes frame time."));
+
 #if !UE_BUILD_SHIPPING
 static TAutoConsoleVariable<int32> CVarPingDisplayServerTime( TEXT( "net.PingDisplayServerTime" ), 0, TEXT( "Show server frame time" ) );
 #endif
@@ -163,6 +165,8 @@ UNetConnection::UNetConnection(const FObjectInitializer& ObjectInitializer)
 ,	TickCount			( 0 )
 ,	LastProcessedFrame	( 0 )
 ,	ConnectTime			( 0.0 )
+,	LastOSReceiveTime	()
+,	bIsOSReceiveTimeLocal(false)
 
 ,	AllowMerge			( false )
 ,	TimeSensitive		( false )
@@ -1703,8 +1707,24 @@ bool UNetConnection::ReadPacketInfo(FBitReader& Reader)
 		}
 #endif
 
+		double PacketReceiveTime = 0.0;
+		FTimespan& RecvTimespan = LastOSReceiveTime.Timestamp;
+
+		if (!RecvTimespan.IsZero() && Driver != nullptr && CVarPingUsePacketRecvTime.GetValueOnAnyThread())
+		{
+			if (bIsOSReceiveTimeLocal)
+			{
+				PacketReceiveTime = RecvTimespan.GetTotalSeconds();
+			}
+			else if (ISocketSubsystem* SocketSubsystem = Driver->GetSocketSubsystem())
+			{
+				PacketReceiveTime = SocketSubsystem->TranslatePacketTimestamp(LastOSReceiveTime);
+			}
+		}
+
+
 		// use FApp's time because it is set closer to the beginning of the frame - we don't care about the time so far of the current frame to process the packet
-		const double CurrentTime = FApp::GetCurrentTime();
+		const double CurrentTime = (PacketReceiveTime != 0.0 ? PacketReceiveTime : FApp::GetCurrentTime());
 		const double GameTime	 = ServerFrameTime;
 		const double RTT		 = (CurrentTime - OutLagTime[Index] ) - ( CVarPingExcludeFrameTime.GetValueOnAnyThread() ? GameTime : 0.0 );
 		const double NewLag		 = FMath::Max( RTT, 0.0 );
