Make sure VoipSynthListenerComponent starts back up if is stopped by audio engine

diff --git a/Engine/Plugins/Online/OnlineSubsystemUtils/Source/OnlineSubsystemUtils/Private/VoiceEngineImpl.cpp b/Engine/Plugins/Online/OnlineSubsystemUtils/Source/OnlineSubsystemUtils/Private/VoiceEngineImpl.cpp
--- a/Engine/Plugins/Online/OnlineSubsystemUtils/Source/OnlineSubsystemUtils/Private/VoiceEngineImpl.cpp
+++ b/Engine/Plugins/Online/OnlineSubsystemUtils/Source/OnlineSubsystemUtils/Private/VoiceEngineImpl.cpp
@@ -590,6 +590,12 @@ uint32 FVoiceEngineImpl::SubmitRemoteVoiceData(const FUniqueNetIdWrapper& Remote
 		}
 
 		QueuedData.VoipSynthComponent->SubmitPacket((float*)DecompressedVoiceBuffer.GetData(), BytesWritten, InSampleCount, EVoipStreamDataFormat::Int16);
+
+		// Try to start the VoipSynthComponent if it has been killed by the audio engine.
+		if (!QueuedData.VoipSynthComponent->IsPlaying())
+		{
+			QueuedData.VoipSynthComponent->Start();
+		}
 	}
 
 	return ONLINE_SUCCESS;
