[NO NON-PII COMMIT MESSAGE]

diff --git a/Engine/Plugins/Online/OnlineSubsystemUtils/Source/OnlineSubsystemUtils/Private/OnlineSubsystemUtils.cpp b/Engine/Plugins/Online/OnlineSubsystemUtils/Source/OnlineSubsystemUtils/Private/OnlineSubsystemUtils.cpp
--- a/Engine/Plugins/Online/OnlineSubsystemUtils/Source/OnlineSubsystemUtils/Private/OnlineSubsystemUtils.cpp
+++ b/Engine/Plugins/Online/OnlineSubsystemUtils/Source/OnlineSubsystemUtils/Private/OnlineSubsystemUtils.cpp
@@ -37,6 +37,14 @@
 #include "Tests/TestExternalUIInterface.h"
 #include "Tests/TestPresenceInterface.h"
 
+static int32 CvarAlwaysPlayVoipComponent = 1;
+FAutoConsoleVariableRef CVarAlwaysPlayVoipComponent(
+	TEXT("au.voip.AlwaysPlayVoiceComponent"),
+	CvarAlwaysPlayVoipComponent,
+	TEXT("When set to 1, guarantees that voip components won't get deprioritized. \n")
+	TEXT("0: Let voip components get killed, 1: force VOIP components to be higher priority than all other audio sources."),
+	ECVF_Default);
+
 UAudioComponent* CreateVoiceAudioComponent(uint32 SampleRate, int32 NumChannels)
 {
 	UAudioComponent* AudioComponent = nullptr;
@@ -100,6 +108,11 @@ UVoipListenerSynthComponent* CreateVoiceSynthComponent(uint32 SampleRate)
 				SynthComponentPtr->SoundClass = LoadObject<USoundClass>(nullptr, *VoiPSoundClassName.ToString());
 			}
 
+			if (CvarAlwaysPlayVoipComponent)
+			{
+				SynthComponentPtr->bAlwaysPlay = true;
+			}
+
 			SynthComponentPtr->Initialize(SampleRate);
 		}
 		else
diff --git a/Engine/Source/Runtime/AudioMixer/Private/Components/SynthComponent.cpp b/Engine/Source/Runtime/AudioMixer/Private/Components/SynthComponent.cpp
--- a/Engine/Source/Runtime/AudioMixer/Private/Components/SynthComponent.cpp
+++ b/Engine/Source/Runtime/AudioMixer/Private/Components/SynthComponent.cpp
@@ -121,7 +121,7 @@ USynthComponent::USynthComponent(const FObjectInitializer& ObjectInitializer)
 	bIsSynthPlaying = false;
 	bIsInitialized = false;
 	bIsUISound = false;
-
+	bAlwaysPlay = false;
 	Synth = nullptr;
 
 	// Set the default sound class
@@ -288,6 +288,7 @@ void USynthComponent::CreateAudioComponent()
 		// Set defaults to be the same as audio component defaults
 		AudioComponent->EnvelopeFollowerAttackTime = EnvelopeFollowerAttackTime;
 		AudioComponent->EnvelopeFollowerReleaseTime = EnvelopeFollowerReleaseTime;
+		AudioComponent->bAlwaysPlay = bAlwaysPlay;
 	}
 }
 
diff --git a/Engine/Source/Runtime/AudioMixer/Public/Components/SynthComponent.h b/Engine/Source/Runtime/AudioMixer/Public/Components/SynthComponent.h
--- a/Engine/Source/Runtime/AudioMixer/Public/Components/SynthComponent.h
+++ b/Engine/Source/Runtime/AudioMixer/Public/Components/SynthComponent.h
@@ -184,6 +184,9 @@ public:
 	UPROPERTY()
 	uint8 bIsPreviewSound : 1;
 
+	/** Whether to artificially prioritize the component to play */
+	uint8 bAlwaysPlay : 1;
+
 	/** Call if creating this synth component not via an actor component in BP, but in code or some other location. Optionally override the sample rate of the sound wave, otherwise it uses the audio device's sample rate. */
 	void Initialize(int32 SampleRateOverride = INDEX_NONE);
 
