Fix address synthesis when given a non-platform address. Now synthesis will be attempted on any input address if necessary or the address needs resolution.

A more robust solution will come later.

diff --git a/Engine/Plugins/Online/OnlineSubsystemUtils/Source/OnlineSubsystemUtils/Private/IpConnection.cpp b/Engine/Plugins/Online/OnlineSubsystemUtils/Source/OnlineSubsystemUtils/Private/IpConnection.cpp
--- a/Engine/Plugins/Online/OnlineSubsystemUtils/Source/OnlineSubsystemUtils/Private/IpConnection.cpp
+++ b/Engine/Plugins/Online/OnlineSubsystemUtils/Source/OnlineSubsystemUtils/Private/IpConnection.cpp
@@ -72,7 +72,7 @@ void UIpConnection::InitLocalConnection(UNetDriver* InDriver, class FSocket* InS
 	RemoteAddr->SetIp(*InURL.Host, bIsValid);
 
 	// If the protocols do not match, attempt to synthesize the address so they do.
-	if (bIsValid && InSocket->GetProtocol() != RemoteAddr->GetProtocolType())
+	if ((bIsValid && InSocket->GetProtocol() != RemoteAddr->GetProtocolType()) || !bIsValid)
 	{
 		SCOPE_CYCLE_COUNTER(STAT_IpConnection_AddressSynthesis);
 
@@ -84,6 +84,7 @@ void UIpConnection::InitLocalConnection(UNetDriver* InDriver, class FSocket* InS
 		if (MapRequest.ReturnCode == SE_NO_ERROR && MapRequest.Results.Num() > 0)
 		{
 			RemoteAddr = MapRequest.Results[0].Address->Clone();
+			bIsValid = true;
 		}
 		else
 		{
