Speculative fix for GetLasErrorCode() socket crash.

diff --git a/Engine/Plugins/Online/OnlineSubsystemUtils/Source/OnlineSubsystemUtils/Private/IpNetDriver.cpp b/Engine/Plugins/Online/OnlineSubsystemUtils/Source/OnlineSubsystemUtils/Private/IpNetDriver.cpp
--- a/Engine/Plugins/Online/OnlineSubsystemUtils/Source/OnlineSubsystemUtils/Private/IpNetDriver.cpp
+++ b/Engine/Plugins/Online/OnlineSubsystemUtils/Source/OnlineSubsystemUtils/Private/IpNetDriver.cpp
@@ -391,11 +391,11 @@ private:
 		RecvMultiIdx = 0;
 		RecvMultiPacketCount = 0;
 
-		while (true)
+		while (Socket != nullptr)
 		{
 			SCOPE_CYCLE_COUNTER(STAT_IpNetDriver_RecvFromSocket);
 
-			bool bRecvMultiOk = Socket != nullptr ? Socket->RecvMulti(*RMState) : false;
+			bool bRecvMultiOk = Socket->RecvMulti(*RMState);
 
 			if (!bRecvMultiOk)
 			{
