Update GetLastErrorCode() check.

diff --git a/Engine/Plugins/Online/OnlineSubsystemUtils/Source/OnlineSubsystemUtils/Private/IpNetDriver.cpp b/Engine/Plugins/Online/OnlineSubsystemUtils/Source/OnlineSubsystemUtils/Private/IpNetDriver.cpp
--- a/Engine/Plugins/Online/OnlineSubsystemUtils/Source/OnlineSubsystemUtils/Private/IpNetDriver.cpp
+++ b/Engine/Plugins/Online/OnlineSubsystemUtils/Source/OnlineSubsystemUtils/Private/IpNetDriver.cpp
@@ -391,7 +391,9 @@ private:
 		RecvMultiIdx = 0;
 		RecvMultiPacketCount = 0;
 
-		while (Socket != nullptr)
+		bBreak = Socket == nullptr;
+
+		while (!bBreak)
 		{
 			SCOPE_CYCLE_COUNTER(STAT_IpNetDriver_RecvFromSocket);
 
