Refactored more TickDispatch code into FPacketIterator, and tweaked receive time limit check packet counting.

diff --git a/Engine/Plugins/Online/OnlineSubsystemUtils/Source/OnlineSubsystemUtils/Private/IpNetDriver.cpp b/Engine/Plugins/Online/OnlineSubsystemUtils/Source/OnlineSubsystemUtils/Private/IpNetDriver.cpp
--- a/Engine/Plugins/Online/OnlineSubsystemUtils/Source/OnlineSubsystemUtils/Private/IpNetDriver.cpp
+++ b/Engine/Plugins/Online/OnlineSubsystemUtils/Source/OnlineSubsystemUtils/Private/IpNetDriver.cpp
@@ -100,7 +100,8 @@ TAutoConsoleVariable<int32> CVarNetUseRecvTimestamps(
  *
  * Encapsulates the NetDriver TickDispatch code required for executing all variations of packet receives
  * (FSocket::RecvFrom, FSocket::RecvMulti, and the Receive Thread),
- * as well as implementing/abstracting-away some of the outermost (non-NetConnection-related) parts of the DDoS detection code.
+ * as well as implementing/abstracting-away some of the outermost (non-NetConnection-related) parts of the DDoS detection code,
+ * and code for timing receives/iterations (which affects control flow).
  */
 class FPacketIterator
 {
@@ -124,22 +125,35 @@ private:
 		/** Error if receiving a packet failed */
 		ESocketErrors Error;
 	};
-	
+
+
 private:
 	FPacketIterator(UIpNetDriver* InDriver)
+		: FPacketIterator(InDriver, InDriver->RecvMultiState.Get(), FPlatformTime::Seconds(),
+							(InDriver->MaxSecondsInReceive > 0.0 && InDriver->NbPacketsBetweenReceiveTimeTest > 0))
+	{
+	}
+
+	FPacketIterator(UIpNetDriver* InDriver, FRecvMulti* InRMState, double InStartReceiveTime, bool bInCheckReceiveTime)
 		: bBreak(false)
+		, IterationCount(0)
 		, Driver(InDriver)
-		, DDoS(Driver->DDoS)
-		, SocketSubsystem(Driver->GetSocketSubsystem())
-		, Socket(Driver->Socket)
-		, SocketReceiveThreadRunnable(Driver->SocketReceiveThreadRunnable.Get())
+		, DDoS(InDriver->DDoS)
+		, SocketSubsystem(InDriver->GetSocketSubsystem())
+		, Socket(InDriver->Socket)
+		, SocketReceiveThreadRunnable(InDriver->SocketReceiveThreadRunnable.Get())
 		, CurrentPacket()
+		, RMState(InRMState)
+		, bUseRecvMulti(CVarNetUseRecvMulti.GetValueOnAnyThread() != 0 && InRMState != nullptr)
 		, RecvMultiIdx(0)
 		, RecvMultiPacketCount(0)
+		, StartReceiveTime(InStartReceiveTime)
+		, bCheckReceiveTime(bInCheckReceiveTime)
+		, CheckReceiveTimePacketCountMask(bInCheckReceiveTime ? (FMath::RoundUpToPowerOfTwo(InDriver->NbPacketsBetweenReceiveTimeTest)-1) : 0)
+		, BailOutTime(InStartReceiveTime + InDriver->MaxSecondsInReceive)
+		, bSlowFrameChecks(UIpNetDriver::OnNetworkProcessingCausingSlowFrame.IsBound())
+		, AlarmTime(InStartReceiveTime + GIpNetDriverMaxDesiredTimeSliceBeforeAlarmSecs)
 	{
-		RMState = Driver->RecvMultiState.Get();
-		bUseRecvMulti = CVarNetUseRecvMulti.GetValueOnAnyThread() != 0 && RMState != nullptr;
-
 		if (!bUseRecvMulti && SocketSubsystem != nullptr)
 		{
 			CurrentPacket.Address = SocketSubsystem->CreateInternetAddr();
@@ -148,8 +162,19 @@ private:
 		AdvanceCurrentPacket();
 	}
 
+	~FPacketIterator()
+	{
+		const float DeltaReceiveTime = FPlatformTime::Seconds() - StartReceiveTime;
+
+		if (DeltaReceiveTime > GIpNetDriverLongFramePrintoutThresholdSecs)
+		{
+			UE_LOG(LogNet, Warning, TEXT("Took too long to receive packets. Time: %2.2f %s"), DeltaReceiveTime, *Driver->GetName());
+		}
+	}
+
 	FORCEINLINE FPacketIterator& operator++()
 	{
+		IterationCount++;
 		AdvanceCurrentPacket();
 
 		return *this;
@@ -230,22 +255,55 @@ private:
 	 */
 	void AdvanceCurrentPacket()
 	{
-		if (bUseRecvMulti)
+		// @todo: Remove the slow frame checks, eventually - potential DDoS and Switch platform constraint
+		if (bSlowFrameChecks)
 		{
-			if (RecvMultiPacketCount == 0 || ((RecvMultiIdx + 1) >= RecvMultiPacketCount))
+			const double CurrentTime = FPlatformTime::Seconds();
+
+			if (CurrentTime > AlarmTime)
 			{
-				AdvanceRecvMultiState();
+				Driver->OnNetworkProcessingCausingSlowFrame.Broadcast();
+
+				AlarmTime = CurrentTime + GIpNetDriverMaxDesiredTimeSliceBeforeAlarmSecs;
 			}
-			else
+		}
+
+		if (bCheckReceiveTime)
+		{
+			if ((IterationCount & CheckReceiveTimePacketCountMask) == 0 && IterationCount > 0)
 			{
-				RecvMultiIdx++;
-			}
+				const double CurrentTime = FPlatformTime::Seconds();
+
+				if (CurrentTime > BailOutTime)
+				{
+					// NOTE: For RecvMulti, this will mass-dump packets, leading to packetloss. Avoid using with RecvMulti.
+					bBreak = true;
 
-			// At this point, bBreak will be set, or RecvMultiPacketCount will be > 0
+					UE_LOG(LogNet, Warning, TEXT("Stopping packet reception after processing for more than %f seconds. %s"),
+							Driver->MaxSecondsInReceive, *Driver->GetName());
+				}
+			}
 		}
-		else
+
+		if (!bBreak)
 		{
-			bBreak = !ReceiveSinglePacket();
+			if (bUseRecvMulti)
+			{
+				if (RecvMultiPacketCount == 0 || ((RecvMultiIdx + 1) >= RecvMultiPacketCount))
+				{
+					AdvanceRecvMultiState();
+				}
+				else
+				{
+					RecvMultiIdx++;
+				}
+
+				// At this point, bBreak will be set, or RecvMultiPacketCount will be > 0
+			}
+			else
+			{
+				bBreak = !ReceiveSinglePacket();
+			}
 		}
 	}
 
@@ -449,33 +507,54 @@ private:
 	/** Specified internally, when the packet iterator should break/stop (no packets, DDoS limits triggered, etc.) */
 	bool bBreak;
 
+	/** The number of packets iterated thus far */
+	int64 IterationCount;
+
 
 	/** Cached reference to the NetDriver, and NetDriver variables/values */
 
-	UIpNetDriver* Driver;
+	UIpNetDriver* const Driver;
 
 	FDDoSDetection& DDoS;
 
-	ISocketSubsystem* SocketSubsystem;
+	ISocketSubsystem* const SocketSubsystem;
 
 	FSocket*& Socket;
 
-	UIpNetDriver::FReceiveThreadRunnable* SocketReceiveThreadRunnable;
+	FReceiveThreadRunnable* const SocketReceiveThreadRunnable;
 
 	/** Stores information for the current packet being received (when using single-receive mode) */
 	FCachedPacket CurrentPacket;
 
 	/** Stores information for receiving packets using RecvMulti */
-	FRecvMulti* RMState;
+	FRecvMulti* const RMState;
 
 	/** Whether or not RecvMulti is enabled/supported */
-	bool bUseRecvMulti;
+	const bool bUseRecvMulti;
 
 	/** The RecvMulti index of the next packet to be received (if RecvMultiPacketCount > 0) */
 	int32 RecvMultiIdx;
 
 	/** The number of packets waiting to be read from the FRecvMulti state */
 	int32 RecvMultiPacketCount;
+
+	/** The time at which packet iteration/receiving began */
+	const double StartReceiveTime;
+
+	/** Whether or not to perform receive time limit checks */
+	const bool bCheckReceiveTime;
+
+	/** Receive time is checked every 'x' number of packets, with this mask used to count the packets ('x' is a power of 2) */
+	const int32 CheckReceiveTimePacketCountMask;
+
+	/** The time at which to bail out of the receive loop, if it's time limited */
+	const double BailOutTime;
+
+	/** Whether or not checks for slow frames are active */
+	const bool bSlowFrameChecks;
+
+	/** Cached time at which to trigger a slow frame alarm */
+	double AlarmTime;
 };
 
 /**
@@ -734,48 +813,11 @@ void UIpNetDriver::TickDispatch(float DeltaTime)
 	DDoS.PreFrameReceive(DeltaTime);
 
 	ISocketSubsystem* SocketSubsystem = GetSocketSubsystem();
-	const double StartReceiveTime = FPlatformTime::Seconds();
-	double AlarmTime = StartReceiveTime + GIpNetDriverMaxDesiredTimeSliceBeforeAlarmSecs;
-	const bool bSlowFrameChecks = OnNetworkProcessingCausingSlowFrame.IsBound();
 	bool bRetrieveTimestamps = CVarNetUseRecvTimestamps.GetValueOnAnyThread() != 0;
 
-	const bool bCheckReceiveTime = (MaxSecondsInReceive > 0.0) && (NbPacketsBetweenReceiveTimeTest > 0);
-	const double BailOutTime = StartReceiveTime + MaxSecondsInReceive;
-	int32 PacketsLeftUntilTimeTest = NbPacketsBetweenReceiveTimeTest;
-
-	bool bContinueProcessing(true);
-
 	// Process all incoming packets
-	for (FPacketIterator It(this); It && bContinueProcessing; ++It)
+	for (FPacketIterator It(this); It; ++It)
 	{
-		// @todo: Remove the slow frame checks, eventually - potential DDoS and Switch platform constraint
-		if (bSlowFrameChecks)
-		{
-			const double CurrentTime = FPlatformTime::Seconds();
-			if (CurrentTime > AlarmTime)
-			{
-				OnNetworkProcessingCausingSlowFrame.Broadcast();
-
-				AlarmTime = CurrentTime + GIpNetDriverMaxDesiredTimeSliceBeforeAlarmSecs;
-			}
-		}
-
-		if (bCheckReceiveTime)
-		{
-			--PacketsLeftUntilTimeTest;
-			if (PacketsLeftUntilTimeTest <= 0)
-			{
-				PacketsLeftUntilTimeTest = NbPacketsBetweenReceiveTimeTest;
-
-				const double CurrentTime = FPlatformTime::Seconds();
-				if (CurrentTime > BailOutTime)
-				{
-					bContinueProcessing = false;
-					UE_LOG(LogNet, Warning, TEXT("UIpNetDriver::TickDispatch: Stopping packet reception after processing for more than %f seconds. %s"), MaxSecondsInReceive, *GetName());
-				}
-			}
-		}
-
 		FReceivedPacketView ReceivedPacket;
 		bool bOk = It.GetCurrentPacket(ReceivedPacket);
 		const TSharedRef<const FInternetAddr> FromAddr = ReceivedPacket.Address.ToSharedRef();
@@ -1000,13 +1042,6 @@ void UIpNetDriver::TickDispatch(float DeltaTime)
 	}
 
 	DDoS.PostFrameReceive();
-
-	const float DeltaReceiveTime = FPlatformTime::Seconds() - StartReceiveTime;
-
-	if (DeltaReceiveTime > GIpNetDriverLongFramePrintoutThresholdSecs)
-	{
-		UE_LOG( LogNet, Warning, TEXT( "UIpNetDriver::TickDispatch: Took too long to receive packets. Time: %2.2f %s" ), DeltaReceiveTime, *GetName() );
-	}
 }
 
 UNetConnection* UIpNetDriver::ProcessConnectionlessPacket(FReceivedPacketView& PacketRef, const FPacketBufferView& WorkingBuffer)
