Restore proper type for refactored code.

diff --git a/Engine/Plugins/Online/OnlineSubsystemUtils/Source/OnlineSubsystemUtils/Private/IpNetDriver.cpp b/Engine/Plugins/Online/OnlineSubsystemUtils/Source/OnlineSubsystemUtils/Private/IpNetDriver.cpp
--- a/Engine/Plugins/Online/OnlineSubsystemUtils/Source/OnlineSubsystemUtils/Private/IpNetDriver.cpp
+++ b/Engine/Plugins/Online/OnlineSubsystemUtils/Source/OnlineSubsystemUtils/Private/IpNetDriver.cpp
@@ -521,7 +521,7 @@ private:
 
 	FSocket*& Socket;
 
-	FReceiveThreadRunnable* const SocketReceiveThreadRunnable;
+	UIpNetDriver::FReceiveThreadRunnable* const SocketReceiveThreadRunnable;
 
 	/** Stores information for the current packet being received (when using single-receive mode) */
 	FCachedPacket CurrentPacket;
