Fix buffer overrun in FVoiceEngineImpl

diff --git a/Engine/Plugins/Online/OnlineSubsystemUtils/Source/OnlineSubsystemUtils/Private/VoiceEngineImpl.cpp b/Engine/Plugins/Online/OnlineSubsystemUtils/Source/OnlineSubsystemUtils/Private/VoiceEngineImpl.cpp
--- a/Engine/Plugins/Online/OnlineSubsystemUtils/Source/OnlineSubsystemUtils/Private/VoiceEngineImpl.cpp
+++ b/Engine/Plugins/Online/OnlineSubsystemUtils/Source/OnlineSubsystemUtils/Private/VoiceEngineImpl.cpp
@@ -456,10 +456,11 @@ uint32 FVoiceEngineImpl::ReadLocalVoiceData(uint32 LocalUserNum, uint8* Data, ui
 		}
 
 		// Get new uncompressed data
-		uint8* RemainingDecompressedBufferPtr = DecompressedVoiceBuffer.GetData() + PlayerVoiceData[LocalUserNum].VoiceRemainderSize;
+		uint8* const RemainingDecompressedBufferPtr = DecompressedVoiceBuffer.GetData() + PlayerVoiceData[LocalUserNum].VoiceRemainderSize;
+		const uint32 RemainingDecompressedBufferSize = DecompressedVoiceBuffer.Num() - PlayerVoiceData[LocalUserNum].VoiceRemainderSize;
 		uint32 ByteWritten = 0;
 		uint64 NewSampleCount = 0;
-		VoiceResult = VoiceCapture->GetVoiceData(DecompressedVoiceBuffer.GetData() + PlayerVoiceData[LocalUserNum].VoiceRemainderSize, NewVoiceDataBytes, ByteWritten, NewSampleCount);
+		VoiceResult = VoiceCapture->GetVoiceData(RemainingDecompressedBufferPtr, RemainingDecompressedBufferSize, ByteWritten, NewSampleCount);
 		
 		TotalVoiceBytes = ByteWritten + PlayerVoiceData[LocalUserNum].VoiceRemainderSize;
 
