Added debug CVar for reproducing duplicate client packets from two IP's "net.DebugDualIPs".

diff --git a/Engine/Plugins/Online/OnlineSubsystemUtils/Source/OnlineSubsystemUtils/Classes/IpNetDriver.h b/Engine/Plugins/Online/OnlineSubsystemUtils/Source/OnlineSubsystemUtils/Classes/IpNetDriver.h
--- a/Engine/Plugins/Online/OnlineSubsystemUtils/Source/OnlineSubsystemUtils/Classes/IpNetDriver.h
+++ b/Engine/Plugins/Online/OnlineSubsystemUtils/Source/OnlineSubsystemUtils/Classes/IpNetDriver.h
@@ -15,12 +15,21 @@
 #include "SocketTypes.h"
 #include "IpNetDriver.generated.h"
 
+
+// Forward declarations
 class Error;
 class FInternetAddr;
 class FNetworkNotify;
 class FSocket;
 struct FRecvMulti;
 
+
+// CVars
+#if !UE_BUILD_SHIPPING
+extern TSharedPtr<FInternetAddr> GCurrentDuplicateIP;
+#endif
+
+
 UCLASS(transient, config=Engine)
 class ONLINESUBSYSTEMUTILS_API UIpNetDriver : public UNetDriver
 {
diff --git a/Engine/Plugins/Online/OnlineSubsystemUtils/Source/OnlineSubsystemUtils/Private/IpConnection.cpp b/Engine/Plugins/Online/OnlineSubsystemUtils/Source/OnlineSubsystemUtils/Private/IpConnection.cpp
--- a/Engine/Plugins/Online/OnlineSubsystemUtils/Source/OnlineSubsystemUtils/Private/IpConnection.cpp
+++ b/Engine/Plugins/Online/OnlineSubsystemUtils/Source/OnlineSubsystemUtils/Private/IpConnection.cpp
@@ -8,6 +8,7 @@ Notes:
 =============================================================================*/
 
 #include "IpConnection.h"
+#include "IpNetDriver.h"
 #include "SocketSubsystem.h"
 #include "Engine/Engine.h"
 
@@ -17,6 +18,7 @@ Notes:
 #include "Net/DataChannel.h"
 
 #include "Net/Core/Misc/PacketAudit.h"
+#include "Misc/ScopeExit.h"
 
 /*-----------------------------------------------------------------------------
 	Declarations.
@@ -171,6 +173,30 @@ void UIpConnection::LowLevelSend(void* Data, int32 CountBits, FOutPacketTraits&
 {
 	const uint8* DataToSend = reinterpret_cast<uint8*>(Data);
 
+#if !UE_BUILD_SHIPPING
+	TSharedPtr<FInternetAddr> OrigAddr;
+
+	if (GCurrentDuplicateIP.IsValid() && RemoteAddr->CompareEndpoints(*GCurrentDuplicateIP))
+	{
+		OrigAddr = RemoteAddr;
+
+		TSharedRef<FInternetAddr> NewAddr = OrigAddr->Clone();
+		int32 NewPort = NewAddr->GetPort() - 9876;
+
+		NewAddr->SetPort(NewPort >= 0 ? NewPort : (65536 + NewPort));
+
+		RemoteAddr = NewAddr;
+	}
+
+	ON_SCOPE_EXIT
+	{
+		if (OrigAddr.IsValid())
+		{
+			RemoteAddr = OrigAddr;
+		}
+	};
+#endif
+
 	if( ResolveInfo )
 	{
 		// If destination address isn't resolved yet, send nowhere.
diff --git a/Engine/Plugins/Online/OnlineSubsystemUtils/Source/OnlineSubsystemUtils/Private/IpNetDriver.cpp b/Engine/Plugins/Online/OnlineSubsystemUtils/Source/OnlineSubsystemUtils/Private/IpNetDriver.cpp
--- a/Engine/Plugins/Online/OnlineSubsystemUtils/Source/OnlineSubsystemUtils/Private/IpNetDriver.cpp
+++ b/Engine/Plugins/Online/OnlineSubsystemUtils/Source/OnlineSubsystemUtils/Private/IpNetDriver.cpp
@@ -93,6 +93,17 @@ TAutoConsoleVariable<int32> CVarNetUseRecvTimestamps(
 	TEXT("If true and if net.UseRecvMulti is also true, on a Unix/Linux platform, ")
 		TEXT("the kernel timestamp will be retrieved for each packet received, providing more accurate ping calculations."));
 
+#if !UE_BUILD_SHIPPING
+TAutoConsoleVariable<int32> CVarNetDebugDualIPs(
+	TEXT("net.DebugDualIPs"),
+	0,
+	TEXT("If true, will duplicate every packet received, and process with a new (deterministic) IP, ")
+		TEXT("to emulate receiving client packets from dual IP's - which can happen under real-world network conditions")
+		TEXT("(only supports a single client on the server)."));
+
+TSharedPtr<FInternetAddr> GCurrentDuplicateIP;
+#endif
+
 
 
 /**
@@ -143,6 +154,10 @@ private:
 		, Socket(InDriver->Socket)
 		, SocketReceiveThreadRunnable(InDriver->SocketReceiveThreadRunnable.Get())
 		, CurrentPacket()
+#if !UE_BUILD_SHIPPING
+		, bDebugDualIPs(CVarNetDebugDualIPs.GetValueOnAnyThread() != 0)
+		, DuplicatePacket()
+#endif
 		, RMState(InRMState)
 		, bUseRecvMulti(CVarNetUseRecvMulti.GetValueOnAnyThread() != 0 && InRMState != nullptr)
 		, RecvMultiIdx(0)
@@ -159,6 +174,13 @@ private:
 			CurrentPacket.Address = SocketSubsystem->CreateInternetAddr();
 		}
 
+#if !UE_BUILD_SHIPPING
+		if (bDebugDualIPs && !bUseRecvMulti)
+		{
+			DuplicatePacket = MakeUnique<FCachedPacket>();
+		}
+#endif
+
 		AdvanceCurrentPacket();
 	}
 
@@ -209,6 +231,18 @@ private:
 			bRecvSuccess = CurrentPacket.bRecvSuccess;
 		}
 
+#if !UE_BUILD_SHIPPING
+		if (IsDuplicatePacket() && OutPacket.Address.IsValid())
+		{
+			TSharedRef<FInternetAddr> NewAddr = OutPacket.Address->Clone();
+
+			NewAddr->SetPort((NewAddr->GetPort() + 9876) & 0xFFFF);
+
+			OutPacket.Address = NewAddr;
+			GCurrentDuplicateIP = NewAddr;
+		}
+#endif
+
 		return bRecvSuccess;
 	}
 
@@ -287,6 +321,13 @@ private:
 
 		if (!bBreak)
 		{
+#if !UE_BUILD_SHIPPING
+			if (IsDuplicatePacket())
+			{
+				CurrentPacket = *DuplicatePacket;
+			}
+			else
+#endif
 			if (bUseRecvMulti)
 			{
 				if (RecvMultiPacketCount == 0 || ((RecvMultiIdx + 1) >= RecvMultiPacketCount))
@@ -303,6 +344,13 @@ private:
 			else
 			{
 				bBreak = !ReceiveSinglePacket();
+
+#if !UE_BUILD_SHIPPING
+				if (bDebugDualIPs && !bBreak)
+				{
+					(*DuplicatePacket) = CurrentPacket;
+				}
+#endif
 			}
 		}
 	}
@@ -502,6 +550,17 @@ private:
 		}
 	}
 
+#if !UE_BUILD_SHIPPING
+	/**
+	 * Whether or not the current packet being iterated, is a duplicate of the previous packet
+	 */
+	FORCEINLINE bool IsDuplicatePacket() const
+	{
+		// When doing Dual IP debugging, every other packet is a duplicate of the previous packet
+		return bDebugDualIPs && (IterationCount % 2) == 1;
+	}
+#endif
+
 
 private:
 	/** Specified internally, when the packet iterator should break/stop (no packets, DDoS limits triggered, etc.) */
@@ -526,6 +585,14 @@ private:
 	/** Stores information for the current packet being received (when using single-receive mode) */
 	FCachedPacket CurrentPacket;
 
+#if !UE_BUILD_SHIPPING
+	/** Whether or not to enable Dual IP debugging */
+	const bool bDebugDualIPs;
+
+	/** When performing Dual IP tests, a duplicate copy of every packet, is stored here */
+	TUniquePtr<FCachedPacket> DuplicatePacket;
+#endif
+
 	/** Stores information for receiving packets using RecvMulti */
 	FRecvMulti* const RMState;
 
@@ -1215,6 +1282,18 @@ void UIpNetDriver::LowLevelSend(TSharedPtr<const FInternetAddr> Address, void* D
 {
 	if (Address.IsValid() && Address->IsValid())
 	{
+#if !UE_BUILD_SHIPPING
+		if (GCurrentDuplicateIP.IsValid() && Address->CompareEndpoints(*GCurrentDuplicateIP))
+		{
+			TSharedRef<FInternetAddr> NewAddr = Address->Clone();
+			int32 NewPort = NewAddr->GetPort() - 9876;
+
+			NewAddr->SetPort(NewPort >= 0 ? NewPort : (65536 + NewPort));
+
+			Address = NewAddr;
+		}
+#endif
+
 		const uint8* DataToSend = reinterpret_cast<uint8*>(Data);
 
 		if (ConnectionlessHandler.IsValid())
