Add support for additional device routing to our native VOIP engine. This can be done either by using IOnlineVoice::PatchRemoteTalkerOutputToEndpoint(DeviceName), or by using the console command "voice.sendRemoteTalkersToEndpoint <Device Name Substring>"

diff --git a/Engine/Plugins/Online/OnlineSubsystem/Source/Public/Interfaces/VoiceInterface.h b/Engine/Plugins/Online/OnlineSubsystem/Source/Public/Interfaces/VoiceInterface.h
--- a/Engine/Plugins/Online/OnlineSubsystem/Source/Public/Interfaces/VoiceInterface.h
+++ b/Engine/Plugins/Online/OnlineSubsystem/Source/Public/Interfaces/VoiceInterface.h
@@ -253,6 +253,32 @@ public:
 	{
 		return nullptr;
 	};
+
+	/**
+	 * This can be used to patch incoming audio to a different audio device.
+	 * @param InDeviceName- name of the device to patch to.
+	 * @param bMuteInGameOutput- if set to true, the audio going out to the game will be muted.
+	 */
+	virtual bool PatchRemoteTalkerOutputToEndpoint(const FString& InDeviceName, bool bMuteInGameOutput = true)
+	{
+		return false;
+	}
+	
+	/**
+	 * This can be used to patch outgoing audio to a different audio device.
+	 * @param InDeviceName- name of the device to patch to.
+	 */
+	virtual bool PatchLocalTalkerOutputToEndpoint(const FString& InDeviceName)
+	{
+		return false;
+	}
+
+	/**
+	 * When called, disconnects all additional audio endpoints we patched to using PatchVoiceOutputToDevice().
+	 */
+	virtual void DisconnectAllEndpoints()
+	{
+	}
 };
 
 typedef TSharedPtr<IVoiceEngine, ESPMode::ThreadSafe> IVoiceEnginePtr;
@@ -482,6 +508,33 @@ public:
 	{
 		return nullptr;
 	};
+
+	/**
+	 * This can be used to patch incoming audio to a different audio device.
+	 * @param InDeviceName- name of the device to patch to.
+	 * @param bMuteInGameOutput- if set to true, the audio going out to the game will be muted.
+	 */
+	virtual bool PatchRemoteTalkerOutputToEndpoint(const FString& InDeviceName, bool bMuteInGameOutput = true)
+	{
+		return false;
+	}
+
+	/**
+	 * This can be used to patch outgoing audio to a different audio device.
+	 * @param InDeviceName- name of the device to patch to.
+	 */
+	virtual bool PatchLocalTalkerOutputToEndpoint(const FString& InDeviceName)
+	{
+		return false;
+	}
+
+
+	/**
+	 * When called, disconnects all additional audio endpoints we patched to using PatchRemoteTalkerOutputToEndpoint().
+	 */
+	virtual void DisconnectAllEndpoints()
+	{
+	}
 };
 
 typedef TSharedPtr<IOnlineVoice, ESPMode::ThreadSafe> IOnlineVoicePtr;
diff --git a/Engine/Plugins/Online/OnlineSubsystemUtils/Source/OnlineSubsystemUtils/OnlineSubsystemUtils.Build.cs b/Engine/Plugins/Online/OnlineSubsystemUtils/Source/OnlineSubsystemUtils/OnlineSubsystemUtils.Build.cs
--- a/Engine/Plugins/Online/OnlineSubsystemUtils/Source/OnlineSubsystemUtils/OnlineSubsystemUtils.Build.cs
+++ b/Engine/Plugins/Online/OnlineSubsystemUtils/Source/OnlineSubsystemUtils/OnlineSubsystemUtils.Build.cs
@@ -33,7 +33,8 @@ public class OnlineSubsystemUtils : ModuleRules
 				"PacketHandler",
 				"Json",
 				"AudioMixer",
-				"SignalProcessing"
+				"SignalProcessing",
+				"AudioMixerCore"
 			}
 		);
 
diff --git a/Engine/Plugins/Online/OnlineSubsystemUtils/Source/OnlineSubsystemUtils/Private/OnlineSubsystemUtils.cpp b/Engine/Plugins/Online/OnlineSubsystemUtils/Source/OnlineSubsystemUtils/Private/OnlineSubsystemUtils.cpp
--- a/Engine/Plugins/Online/OnlineSubsystemUtils/Source/OnlineSubsystemUtils/Private/OnlineSubsystemUtils.cpp
+++ b/Engine/Plugins/Online/OnlineSubsystemUtils/Source/OnlineSubsystemUtils/Private/OnlineSubsystemUtils.cpp
@@ -37,6 +37,66 @@
 #include "Tests/TestExternalUIInterface.h"
 #include "Tests/TestPresenceInterface.h"
 
+
+static FAutoConsoleCommand GSendRemoteTalkersToEndpointCommand(
+	TEXT("voice.sendRemoteTalkersToEndpoint"),
+	TEXT("This will send audio output for all incoming voip audio to the named endpoint. if given no arguments, this will route voice output through the game engine."),
+	FConsoleCommandWithArgsDelegate::CreateStatic(
+		[](const TArray< FString >& Args)
+{
+	IOnlineSubsystem* PlatformSubsystem = IOnlineSubsystem::Get();
+	IOnlineVoicePtr PlatformVoiceInterface = PlatformSubsystem ? PlatformSubsystem->GetVoiceInterface() : nullptr;
+
+	check(PlatformVoiceInterface.IsValid());
+
+	if (Args.Num() == 0)
+	{
+		PlatformVoiceInterface->DisconnectAllEndpoints();
+	}
+	else
+	{
+		FString DeviceName;
+		for (int32 Index = 0; Index < Args.Num(); Index++)
+		{
+			DeviceName.Append(Args[Index]);
+			DeviceName.Append(TEXT(" "), 1);
+		}
+
+		PlatformVoiceInterface->PatchRemoteTalkerOutputToEndpoint(DeviceName);
+	}
+})
+);
+
+static FAutoConsoleCommand GSendLocalTalkersToEndpointCommand(
+	TEXT("voice.sendLocalTalkersToEndpoint"),
+	TEXT("This will send audio output for all outgoing voip audio to the named endpoint. if given no arguments, this will disconnect all external endpoints."),
+	FConsoleCommandWithArgsDelegate::CreateStatic(
+		[](const TArray< FString >& Args)
+{
+	IOnlineSubsystem* PlatformSubsystem = IOnlineSubsystem::Get();
+	IOnlineVoicePtr PlatformVoiceInterface = PlatformSubsystem ? PlatformSubsystem->GetVoiceInterface() : nullptr;
+
+	check(PlatformVoiceInterface.IsValid());
+
+
+	if (Args.Num() == 0)
+	{
+		PlatformVoiceInterface->DisconnectAllEndpoints();
+	}
+	else
+	{
+		FString DeviceName;
+		for (int32 Index = 0; Index < Args.Num(); Index++)
+		{
+			DeviceName.Append(Args[Index]);
+			DeviceName.Append(TEXT(" "), 1);
+		}
+
+		PlatformVoiceInterface->PatchLocalTalkerOutputToEndpoint(DeviceName);
+	}
+})
+);
+
 static int32 CvarAlwaysPlayVoipComponent = 1;
 FAutoConsoleVariableRef CVarAlwaysPlayVoipComponent(
 	TEXT("au.voip.AlwaysPlayVoiceComponent"),
diff --git a/Engine/Plugins/Online/OnlineSubsystemUtils/Source/OnlineSubsystemUtils/Private/VoiceEngineImpl.cpp b/Engine/Plugins/Online/OnlineSubsystemUtils/Source/OnlineSubsystemUtils/Private/VoiceEngineImpl.cpp
--- a/Engine/Plugins/Online/OnlineSubsystemUtils/Source/OnlineSubsystemUtils/Private/VoiceEngineImpl.cpp
+++ b/Engine/Plugins/Online/OnlineSubsystemUtils/Source/OnlineSubsystemUtils/Private/VoiceEngineImpl.cpp
@@ -21,6 +21,27 @@ namespace NotificationClient
 }
 #endif 
 
+namespace VoiceEngineUtilities
+{
+	void DownmixBuffer(const float* InAudio, float* OutAudio, int32 NumFrames, int32 InNumChannels, int32 OutNumChannels)
+	{
+		const float Attenuation = FMath::Clamp<float>(((float)OutNumChannels) / InNumChannels, 0.0f, 1.0f);
+		for (int32 FrameIndex = 0; FrameIndex < NumFrames; FrameIndex++)
+		{
+			for (int32 OutChannelIndex = 0; OutChannelIndex < OutNumChannels; OutChannelIndex++)
+			{
+				const int32 OutSampleIndex = FrameIndex * OutNumChannels + OutChannelIndex;
+				OutAudio[OutSampleIndex] = Attenuation * InAudio[FrameIndex * InNumChannels];
+
+				for (int32 InChannelIndex = 1; InChannelIndex < InNumChannels; InChannelIndex++)
+				{
+					const int32 InSampleIndex = FrameIndex * InNumChannels + InChannelIndex;
+					OutAudio[OutSampleIndex] += Attenuation * InAudio[InSampleIndex];
+				}
+			}
+		}
+	}
+}
 
 FRemoteTalkerDataImpl::FRemoteTalkerDataImpl() :
 	MaxUncompressedDataSize(0),
@@ -146,17 +167,17 @@ void FRemoteTalkerDataImpl::Cleanup()
 	VoipSynthComponent = nullptr;
 }
 
-FVoiceEngineImpl ::FVoiceEngineImpl() :
-	OnlineSubsystem(nullptr),
-	VoiceCapture(nullptr),
-	VoiceEncoder(nullptr),
-	OwningUserIndex(INVALID_INDEX),
-	UncompressedBytesAvailable(0),
-	CompressedBytesAvailable(0),
-	AvailableVoiceResult(EVoiceCaptureState::UnInitialized),
-	bPendingFinalCapture(false),
-	bIsCapturing(false),
-	SerializeHelper(nullptr)
+FVoiceEngineImpl ::FVoiceEngineImpl()
+	: OnlineSubsystem(nullptr)
+	, VoiceCapture(nullptr)
+	, VoiceEncoder(nullptr)
+	, OwningUserIndex(INVALID_INDEX)
+	, UncompressedBytesAvailable(0)
+	, CompressedBytesAvailable(0)
+	, AvailableVoiceResult(EVoiceCaptureState::UnInitialized)
+	, bPendingFinalCapture(false)
+	, bIsCapturing(false)
+	, SerializeHelper(nullptr)
 #if PLATFORM_WINDOWS
 	, bAudioDeviceChanged(false)
 #endif
@@ -635,6 +656,9 @@ void FVoiceEngineImpl::Tick(float DeltaTime)
 
 	TickTalkers(DeltaTime);
 
+	// Push any buffered audio to any connected outputs.
+	AllRemoteTalkerAudio.ProcessAudio();
+
 #if PLATFORM_WINDOWS
 	if (bAudioDeviceChanged)
 	{
@@ -813,6 +837,38 @@ Audio::FPatchOutputStrongPtr FVoiceEngineImpl::GetRemoteTalkerOutput()
 	return AllRemoteTalkerAudio.AddNewOutput(4096 * 2, 1.0f);
 }
 
+bool FVoiceEngineImpl::PatchRemoteTalkerOutputToEndpoint(const FString& InDeviceName, bool bMuteInGameOutput /*= true*/)
+{
+	if (bMuteInGameOutput)
+	{
+		static IConsoleVariable* MuteAudioEngineOutputCVar = IConsoleManager::Get().FindConsoleVariable(TEXT("voice.MuteAudioEngineOutput"));
+		check(MuteAudioEngineOutputCVar);
+		MuteAudioEngineOutputCVar->Set(1, ECVF_SetByGameSetting);
+	}
+	
+	TUniquePtr<FVoiceEndpoint>& Endpoint = ExternalEndpoints.Emplace_GetRef(new FVoiceEndpoint(InDeviceName, UVOIPStatics::GetVoiceSampleRate(), DEFAULT_NUM_VOICE_CHANNELS));
+	Audio::FPatchOutputStrongPtr OutputPatch = AllRemoteTalkerAudio.AddNewOutput(4096 * 2, 1.0f);
+	Endpoint->PatchInOutput(OutputPatch);
+	return true;
+}
+
+bool FVoiceEngineImpl::PatchLocalTalkerOutputToEndpoint(const FString& InDeviceName)
+{
+	TUniquePtr<FVoiceEndpoint>& Endpoint = ExternalEndpoints.Emplace_GetRef(new FVoiceEndpoint(InDeviceName, UVOIPStatics::GetVoiceSampleRate(), DEFAULT_NUM_VOICE_CHANNELS));
+	Audio::FPatchOutputStrongPtr OutputPatch = VoiceCapture->GetMicrophoneAudio(4096 * 2, 1.0f);
+	Endpoint->PatchInOutput(OutputPatch);
+	return true;
+}
+
+void FVoiceEngineImpl::DisconnectAllEndpoints()
+{
+	static IConsoleVariable* MuteAudioEngineOutputCVar = IConsoleManager::Get().FindConsoleVariable(TEXT("voice.MuteAudioEngineOutput"));
+	check(MuteAudioEngineOutputCVar);
+	MuteAudioEngineOutputCVar->Set(0, ECVF_SetByGameSetting);
+
+	ExternalEndpoints.Reset();
+}
+
 int32 FVoiceEngineImpl::GetMaxVoiceRemainderSize()
 {
 	return MAX_VOICE_REMAINDER_SIZE;
@@ -872,3 +928,99 @@ void FVoiceEngineImpl::OnDefaultDeviceChanged()
 	TimeDeviceChaned = FPlatformTime::Seconds();
 }
 #endif
+
+FVoiceEndpoint::FVoiceEndpoint(const FString& InEndpointName, float InSampleRate, int32 InNumChannels)
+	: NumChannelsComingIn(InNumChannels)
+{
+	check(GEngine && GEngine->GetAudioDeviceManager());
+
+	IAudioDeviceModule* AudioModule = GEngine->GetAudioDeviceManager()->GetAudioDeviceModule();
+	check(AudioModule);
+
+	PlatformEndpoint.Reset(AudioModule->CreateAudioMixerPlatformInterface());
+
+	if (PlatformEndpoint.IsValid())
+	{
+		bool Result = PlatformEndpoint->InitializeHardware();
+
+		check(Result);
+
+		int32 DeviceIndex = PlatformEndpoint->GetIndexForDevice(InEndpointName);
+
+		if (DeviceIndex == INDEX_NONE)
+		{
+			UE_LOG(LogVoice, Warning, TEXT("Failed to find device %s, using default output device."), *InEndpointName);
+			DeviceIndex = AUDIO_MIXER_DEFAULT_DEVICE_INDEX;
+		}
+
+		int32 NumFrames = PlatformEndpoint->GetNumFrames(1024);
+
+		OpenParams.NumBuffers = 3;
+		OpenParams.NumFrames = NumFrames;
+		OpenParams.OutputDeviceIndex = DeviceIndex;
+		OpenParams.SampleRate = InSampleRate;
+		OpenParams.AudioMixer = this;
+		OpenParams.MaxSources = 0;
+
+		Result = PlatformEndpoint->OpenAudioStream(OpenParams);
+
+		check(Result);
+
+		PlatformEndpoint->PostInitializeHardware();
+
+		PlatformDeviceInfo = PlatformEndpoint->GetPlatformDeviceInfo();
+
+		PlatformEndpoint->StartAudioStream();
+		PlatformEndpoint->FadeIn();
+	}
+}
+
+FVoiceEndpoint::~FVoiceEndpoint()
+{
+	if (PlatformEndpoint.IsValid())
+	{
+		PlatformEndpoint->StopAudioStream();
+		PlatformEndpoint->CloseAudioStream();
+	}
+}
+
+void FVoiceEndpoint::PatchInOutput(Audio::FPatchOutputStrongPtr& InOutput)
+{
+	FScopeLock ScopeLock(&OutputPatchCriticalSection);
+	OutputPatch = InOutput;
+}
+
+bool FVoiceEndpoint::OnProcessAudioStream(Audio::AlignedFloatBuffer& OutputBuffer)
+{
+	FScopeLock ScopeLock(&OutputPatchCriticalSection);
+
+	int32 NumFrames = OutputBuffer.Num() / PlatformDeviceInfo.NumChannels;
+
+	if (OutputPatch.IsValid() && OutputPatch->GetNumSamplesAvailable() >= NumFrames * NumChannelsComingIn)
+	{
+		if (PlatformDeviceInfo.NumChannels != NumChannelsComingIn)
+		{
+			DownmixBuffer.Reset();
+			DownmixBuffer.AddZeroed(NumFrames * NumChannelsComingIn);
+
+			OutputPatch->PopAudio(DownmixBuffer.GetData(), NumFrames * NumChannelsComingIn, false);
+
+			VoiceEngineUtilities::DownmixBuffer(DownmixBuffer.GetData(), OutputBuffer.GetData(), NumFrames, NumChannelsComingIn, PlatformDeviceInfo.NumChannels);
+		}
+		else
+		{
+			OutputPatch->PopAudio(OutputBuffer.GetData(), OutputBuffer.Num(), false);
+		}
+	}
+	else
+	{
+		FMemory::Memzero(OutputBuffer.GetData(), OutputBuffer.Num() * sizeof(float));
+	}
+
+	return true;
+}
+
+void FVoiceEndpoint::OnAudioStreamShutdown()
+{
+	// Nothing to do here.
+}
diff --git a/Engine/Plugins/Online/OnlineSubsystemUtils/Source/OnlineSubsystemUtils/Private/VoiceInterfaceImpl.cpp b/Engine/Plugins/Online/OnlineSubsystemUtils/Source/OnlineSubsystemUtils/Private/VoiceInterfaceImpl.cpp
--- a/Engine/Plugins/Online/OnlineSubsystemUtils/Source/OnlineSubsystemUtils/Private/VoiceInterfaceImpl.cpp
+++ b/Engine/Plugins/Online/OnlineSubsystemUtils/Source/OnlineSubsystemUtils/Private/VoiceInterfaceImpl.cpp
@@ -906,3 +906,35 @@ Audio::FPatchOutputStrongPtr FOnlineVoiceImpl::GetRemoteTalkerOutput()
 		return nullptr;
 	}
 }
+
+bool FOnlineVoiceImpl::PatchRemoteTalkerOutputToEndpoint(const FString& InDeviceName, bool bMuteInGameOutput /*= true*/)
+{
+	if (VoiceEngine.IsValid())
+	{
+		return VoiceEngine->PatchRemoteTalkerOutputToEndpoint(InDeviceName, bMuteInGameOutput);
+	}
+	else
+	{
+		return false;
+	}
+}
+
+bool FOnlineVoiceImpl::PatchLocalTalkerOutputToEndpoint(const FString& InDeviceName)
+{
+	if (VoiceEngine.IsValid())
+	{
+		return VoiceEngine->PatchLocalTalkerOutputToEndpoint(InDeviceName);
+	}
+	else
+	{
+		return false;
+	}
+}
+
+void FOnlineVoiceImpl::DisconnectAllEndpoints()
+{
+	if (VoiceEngine.IsValid())
+	{
+		return VoiceEngine->DisconnectAllEndpoints();
+	}
+}
diff --git a/Engine/Plugins/Online/OnlineSubsystemUtils/Source/OnlineSubsystemUtils/Private/VoipListenerSynthComponent.cpp b/Engine/Plugins/Online/OnlineSubsystemUtils/Source/OnlineSubsystemUtils/Private/VoipListenerSynthComponent.cpp
--- a/Engine/Plugins/Online/OnlineSubsystemUtils/Source/OnlineSubsystemUtils/Private/VoipListenerSynthComponent.cpp
+++ b/Engine/Plugins/Online/OnlineSubsystemUtils/Source/OnlineSubsystemUtils/Private/VoipListenerSynthComponent.cpp
@@ -17,13 +17,20 @@ FAutoConsoleVariableRef CVarDefaultPatchBufferSize(
 	TEXT("Changes the amount of audio we buffer for VOIP patching, in samples.\n"),
 	ECVF_Default);
 
-static int32 DefaultPatchGainCVar = 1.0f;
+static float DefaultPatchGainCVar = 1.0f;
 FAutoConsoleVariableRef CVarDefaultPatchGain(
 	TEXT("voice.DefaultPatchGain"),
 	DefaultPatchGainCVar,
 	TEXT("Changes the default gain of audio patches, in samples.\n"),
 	ECVF_Default);
 
+static int32 MuteAudioEngineOutputCVar = 0;
+FAutoConsoleVariableRef CVarMuteAudioEngineOutput(
+	TEXT("voice.MuteAudioEngineOutput"),
+	MuteAudioEngineOutputCVar,
+	TEXT("When set to a nonzero value, the output for the audio engine will be muted..\n"),
+	ECVF_Default);
+
 
 bool UVoipListenerSynthComponent::Init(int32& SampleRate)
 {
@@ -63,6 +70,10 @@ int32 UVoipListenerSynthComponent::OnGenerateAudio(float* OutAudio, int32 NumSam
 		ExternalSend.PushAudio(OutAudio, NumSamples);
 	}
 	
+	if (MuteAudioEngineOutputCVar)
+	{
+		FMemory::Memzero(OutAudio, NumSamples * sizeof(float));
+	}
 
 	return NumSamples;
 }
diff --git a/Engine/Plugins/Online/OnlineSubsystemUtils/Source/OnlineSubsystemUtils/Public/VoiceEngineImpl.h b/Engine/Plugins/Online/OnlineSubsystemUtils/Source/OnlineSubsystemUtils/Public/VoiceEngineImpl.h
--- a/Engine/Plugins/Online/OnlineSubsystemUtils/Source/OnlineSubsystemUtils/Public/VoiceEngineImpl.h
+++ b/Engine/Plugins/Online/OnlineSubsystemUtils/Source/OnlineSubsystemUtils/Public/VoiceEngineImpl.h
@@ -14,6 +14,7 @@
 #include "VoipListenerSynthComponent.h"
 #include "OnlineSubsystemUtilsPackage.h"
 #include "AudioDevice.h"
+#include "AudioMixer.h"
 #include "DSP/MultithreadedPatching.h"
 
 #include "VoicePacketImpl.h"
@@ -84,10 +85,39 @@ public:
 	TArray<uint8> UncompressedDataQueue;
 	/** Per remote talker voice decoding state */
 	TSharedPtr<IVoiceDecoder> VoiceDecoder;
-	/** Patch splitter to expose incoming audio to multiple outputs. */
-	Audio::FPatchSplitter RemoteVoiceOutput;
 };
 
+/**
+ * Small class that manages an audio endpoint. Used in FVoiceEngineImpl.
+ */
+class FVoiceEndpoint : Audio::IAudioMixer
+{
+public:
+	FVoiceEndpoint(const FString& InEndpointName, float InSampleRate, int32 InNumChannels);
+	virtual ~FVoiceEndpoint();
+
+	void PatchInOutput(Audio::FPatchOutputStrongPtr& InOutput);
+
+
+	// Begin of IAudioMixer overrides
+	bool OnProcessAudioStream(Audio::AlignedFloatBuffer& OutputBuffer) override;
+	void OnAudioStreamShutdown() override;
+	// End of IAudioMixer overrides
+
+private:
+	int32 NumChannelsComingIn;
+	Audio::AlignedFloatBuffer DownmixBuffer;
+
+	TUniquePtr<Audio::IAudioMixerPlatformInterface> PlatformEndpoint;
+
+	Audio::FAudioMixerOpenStreamParams OpenParams;
+	Audio::FAudioPlatformDeviceInfo PlatformDeviceInfo;
+	
+	Audio::FPatchOutputStrongPtr OutputPatch;
+	FCriticalSection OutputPatchCriticalSection;
+};
+
+
 /**
  * Generic implementation of voice engine, using Voice module for capture/codec
  */
@@ -128,9 +158,6 @@ class ONLINESUBSYSTEMUTILS_API FVoiceEngineImpl : public IVoiceEngine, public FS
 	/** Mapping of UniqueIds to the incoming voice data and their audio component */
 	typedef TMap<FUniqueNetIdWrapper, FRemoteTalkerDataImpl> FRemoteTalkerData;
 
-	/** Total mixed down output of all remote player talkers. */
-	Audio::FPatchMixerSplitter MixedDownRemoteTalkerOutput;
-
 	/** Reference to the main online subsystem */
 	IOnlineSubsystem* OnlineSubsystem;
 
@@ -165,6 +192,13 @@ class ONLINESUBSYSTEMUTILS_API FVoiceEngineImpl : public IVoiceEngine, public FS
 	/** Audio taps for the full mixdown of all remote players. */
 	Audio::FPatchMixerSplitter AllRemoteTalkerAudio;
 
+	/**
+	 * Collection of external endpoints that we are sending local or remote audio to. 
+	 * Note that we need to wrap each FVoiceEndpoint in a unique pointer to ensure the FVoiceEndpoint itself isn't moved elsewhere.
+	 * Otherwise, this will cause a crash in FOutputBuffer::MixNextBuffer(), due to AudioMixer->OnProcessAudioStream(); being called on a stale pointer. 
+	 */
+	TArray<TUniquePtr<FVoiceEndpoint>> ExternalEndpoints;
+
 // Get Audio Device Changes on Windows
 #if PLATFORM_WINDOWS
 	FThreadSafeBool bAudioDeviceChanged;
@@ -297,6 +331,15 @@ public:
 	virtual Audio::FPatchOutputStrongPtr GetMicrophoneOutput() override;
 	virtual Audio::FPatchOutputStrongPtr GetRemoteTalkerOutput() override;
 
+
+	virtual bool PatchRemoteTalkerOutputToEndpoint(const FString& InDeviceName, bool bMuteInGameOutput = true) override;
+
+
+	virtual void DisconnectAllEndpoints() override;
+
+
+	virtual bool PatchLocalTalkerOutputToEndpoint(const FString& InDeviceName) override;
+
 private:
 
 	/**
diff --git a/Engine/Plugins/Online/OnlineSubsystemUtils/Source/OnlineSubsystemUtils/Public/VoiceInterfaceImpl.h b/Engine/Plugins/Online/OnlineSubsystemUtils/Source/OnlineSubsystemUtils/Public/VoiceInterfaceImpl.h
--- a/Engine/Plugins/Online/OnlineSubsystemUtils/Source/OnlineSubsystemUtils/Public/VoiceInterfaceImpl.h
+++ b/Engine/Plugins/Online/OnlineSubsystemUtils/Source/OnlineSubsystemUtils/Public/VoiceInterfaceImpl.h
@@ -153,6 +153,10 @@ public:
 	virtual FString GetVoiceDebugState() const override;
 	virtual Audio::FPatchOutputStrongPtr GetMicrophoneOutput() override;
 	virtual Audio::FPatchOutputStrongPtr GetRemoteTalkerOutput() override;
+	virtual bool PatchRemoteTalkerOutputToEndpoint(const FString& InDeviceName, bool bMuteInGameOutput = true) override;
+	virtual bool PatchLocalTalkerOutputToEndpoint(const FString& InDeviceName) override;
+	virtual void DisconnectAllEndpoints() override;
+	//~IOnlineVoice
 };
 
 typedef TSharedPtr<FOnlineVoiceImpl, ESPMode::ThreadSafe> FOnlineVoiceImplPtr;
diff --git a/Engine/Source/Runtime/AudioMixerCore/Public/AudioMixer.h b/Engine/Source/Runtime/AudioMixerCore/Public/AudioMixer.h
--- a/Engine/Source/Runtime/AudioMixerCore/Public/AudioMixer.h
+++ b/Engine/Source/Runtime/AudioMixerCore/Public/AudioMixer.h
@@ -388,6 +388,18 @@ namespace Audio
 		/** Gets the device information of the given device index. */
 		virtual bool GetOutputDeviceInfo(const uint32 InDeviceIndex, FAudioPlatformDeviceInfo& OutInfo) = 0;
 
+		/**
+		 * Returns the name of the currently used audio device.
+		 */
+		virtual FString GetCurrentDeviceName() { return CurrentDeviceName; }
+
+		/**
+		 * Can be used to look up the current index for a given device name.
+		 * On most platforms, this index may be invalidated if any devices are added or removed.
+		 * Returns INDEX_NONE if 
+		 */
+		virtual int32 GetIndexForDevice(const FString& InDeviceName) { return INDEX_NONE; }
+
 		/** Gets the platform specific audio settings. */
 		virtual FAudioPlatformSettings GetPlatformSettings() const = 0;
 
@@ -579,6 +591,9 @@ namespace Audio
 		/** Source param used to fade in and out audio device. */
 		FParam FadeParam;
 
+		/** This device name can be used to override the default device being used on platforms that use strings to identify audio devices. */
+		FString CurrentDeviceName;
+
 		/** String containing the last generated error. */
 		FString LastError;
 
diff --git a/Engine/Source/Runtime/Engine/Private/AudioDeviceManager.cpp b/Engine/Source/Runtime/Engine/Private/AudioDeviceManager.cpp
--- a/Engine/Source/Runtime/Engine/Private/AudioDeviceManager.cpp
+++ b/Engine/Source/Runtime/Engine/Private/AudioDeviceManager.cpp
@@ -253,6 +253,11 @@ bool FAudioDeviceManager::IsUsingAudioMixer() const
 	return bUsingAudioMixer;
 }
 
+IAudioDeviceModule* FAudioDeviceManager::GetAudioDeviceModule()
+{
+	return AudioDeviceModule;
+}
+
 bool FAudioDeviceManager::Initialize()
 {
 	if (LoadDefaultAudioDeviceModule())
diff --git a/Engine/Source/Runtime/Engine/Public/AudioDeviceManager.h b/Engine/Source/Runtime/Engine/Public/AudioDeviceManager.h
--- a/Engine/Source/Runtime/Engine/Public/AudioDeviceManager.h
+++ b/Engine/Source/Runtime/Engine/Public/AudioDeviceManager.h
@@ -68,6 +68,11 @@ public:
 		FCreateAudioDeviceResults();
 	};
 
+	/**
+	 * returns the currently used audio device module for this platform.
+	 * returns nullptr if Initialize() has not been called yet.
+	 */
+	IAudioDeviceModule* GetAudioDeviceModule();
 
 	/**
 	* Creates and audio device instance internally and returns a
diff --git a/Engine/Source/Runtime/Linux/AudioMixerSDL/Private/AudioMixerPlatformSDL.cpp b/Engine/Source/Runtime/Linux/AudioMixerSDL/Private/AudioMixerPlatformSDL.cpp
--- a/Engine/Source/Runtime/Linux/AudioMixerSDL/Private/AudioMixerPlatformSDL.cpp
+++ b/Engine/Source/Runtime/Linux/AudioMixerSDL/Private/AudioMixerPlatformSDL.cpp
@@ -251,11 +251,6 @@ namespace Audio
 		return true;
 	}
 
-	FString FMixerPlatformSDL::GetCurrentDeviceName() const
-	{
-		return {};
-	}
-
 	bool FMixerPlatformSDL::CloseAudioStream()
 	{
 		if (AudioStreamInfo.StreamState == EAudioOutputStreamState::Closed)
diff --git a/Engine/Source/Runtime/Linux/AudioMixerSDL/Private/AudioMixerPlatformSDL.h b/Engine/Source/Runtime/Linux/AudioMixerSDL/Private/AudioMixerPlatformSDL.h
--- a/Engine/Source/Runtime/Linux/AudioMixerSDL/Private/AudioMixerPlatformSDL.h
+++ b/Engine/Source/Runtime/Linux/AudioMixerSDL/Private/AudioMixerPlatformSDL.h
@@ -45,8 +45,6 @@ namespace Audio
 
 		void HandleOnBufferEnd(uint8* InOutputBuffer, int32 InOutputBufferLength);
 
-		virtual FString GetCurrentDeviceName() const;
-
 	protected:
 		FCriticalSection OutputBufferMutex;
 
diff --git a/Engine/Source/Runtime/SignalProcessing/Private/MultithreadedPatching.cpp b/Engine/Source/Runtime/SignalProcessing/Private/MultithreadedPatching.cpp
--- a/Engine/Source/Runtime/SignalProcessing/Private/MultithreadedPatching.cpp
+++ b/Engine/Source/Runtime/SignalProcessing/Private/MultithreadedPatching.cpp
@@ -95,6 +95,11 @@ namespace Audio
 	}
 
 
+	int32 FPatchOutput::GetNumSamplesAvailable()
+	{
+		return InternalBuffer.Num();
+	}
+
 	FPatchInput::FPatchInput(const FPatchOutputStrongPtr& InOutput)
 		: OutputHandle(InOutput)
 		, PushCallsCounter(0)
diff --git a/Engine/Source/Runtime/SignalProcessing/Public/DSP/MultithreadedPatching.h b/Engine/Source/Runtime/SignalProcessing/Public/DSP/MultithreadedPatching.h
--- a/Engine/Source/Runtime/SignalProcessing/Public/DSP/MultithreadedPatching.h
+++ b/Engine/Source/Runtime/SignalProcessing/Public/DSP/MultithreadedPatching.h
@@ -39,6 +39,9 @@ namespace Audio
 		/** Sums the minimum of NumSamples or however many samples are available into OutBuffer. Returns the number of samples summed into OutBuffer. */
 		int32 MixInAudio(float* OutBuffer, int32 NumSamples, bool bUseLatestAudio);
 
+		/** Returns the current number of samples buffered on this output. */
+		int32 GetNumSamplesAvailable();
+
 		/** Returns true if the input for this patch has been destroyed. */
 		bool IsInputStale() const;
 
diff --git a/Engine/Source/Runtime/Windows/AudioMixerXAudio2/Private/AudioMixerPlatformXAudio2.cpp b/Engine/Source/Runtime/Windows/AudioMixerXAudio2/Private/AudioMixerPlatformXAudio2.cpp
--- a/Engine/Source/Runtime/Windows/AudioMixerXAudio2/Private/AudioMixerPlatformXAudio2.cpp
+++ b/Engine/Source/Runtime/Windows/AudioMixerXAudio2/Private/AudioMixerPlatformXAudio2.cpp
@@ -524,6 +524,33 @@ namespace Audio
 		return true;
 	}
 
+	int32 FMixerPlatformXAudio2::GetIndexForDevice(const FString& InDeviceName)
+	{
+		uint32 TotalNumDevices = 0;
+
+		if (!GetNumOutputDevices(TotalNumDevices))
+		{
+			return INDEX_NONE;
+		}
+
+		// Iterate through every device and see if
+		for (uint32 DeviceIndex = 0; DeviceIndex < TotalNumDevices; DeviceIndex++)
+		{
+			FAudioPlatformDeviceInfo DeviceInfo;
+			if (GetOutputDeviceInfo(DeviceIndex, DeviceInfo))
+			{
+				// check if the device name matches the input device name:
+				if (DeviceInfo.Name.Contains(InDeviceName))
+				{
+					return DeviceIndex;
+				}
+			}
+		}
+
+		// If we've made it here, we weren't able to find a matching device.
+		return INDEX_NONE;
+	}
+
 	bool FMixerPlatformXAudio2::GetDefaultOutputDeviceIndex(uint32& OutDefaultDeviceIndex) const
 	{
 		OutDefaultDeviceIndex = 0;
diff --git a/Engine/Source/Runtime/Windows/AudioMixerXAudio2/Private/AudioMixerPlatformXAudio2.h b/Engine/Source/Runtime/Windows/AudioMixerXAudio2/Private/AudioMixerPlatformXAudio2.h
--- a/Engine/Source/Runtime/Windows/AudioMixerXAudio2/Private/AudioMixerPlatformXAudio2.h
+++ b/Engine/Source/Runtime/Windows/AudioMixerXAudio2/Private/AudioMixerPlatformXAudio2.h
@@ -53,6 +53,7 @@ namespace Audio
 		virtual bool GetNumOutputDevices(uint32& OutNumOutputDevices) override;
 		virtual bool GetOutputDeviceInfo(const uint32 InDeviceIndex, FAudioPlatformDeviceInfo& OutInfo) override;
 		virtual bool GetDefaultOutputDeviceIndex(uint32& OutDefaultDeviceIndex) const override;
+		virtual int32 GetIndexForDevice(const FString& InDeviceName) override;
 		virtual bool OpenAudioStream(const FAudioMixerOpenStreamParams& Params) override;
 		virtual bool CloseAudioStream() override;
 		virtual bool StartAudioStream() override;
