Fix for compile error VoiceEngineImpl.cpp from merge

diff --git a/Engine/Plugins/Online/OnlineSubsystemUtils/Source/OnlineSubsystemUtils/Private/VoiceEngineImpl.cpp b/Engine/Plugins/Online/OnlineSubsystemUtils/Source/OnlineSubsystemUtils/Private/VoiceEngineImpl.cpp
--- a/Engine/Plugins/Online/OnlineSubsystemUtils/Source/OnlineSubsystemUtils/Private/VoiceEngineImpl.cpp
+++ b/Engine/Plugins/Online/OnlineSubsystemUtils/Source/OnlineSubsystemUtils/Private/VoiceEngineImpl.cpp
@@ -846,7 +846,7 @@ bool FVoiceEngineImpl::PatchRemoteTalkerOutputToEndpoint(const FString& InDevice
 		MuteAudioEngineOutputCVar->Set(1, ECVF_SetByGameSetting);
 	}
 	
-	TUniquePtr<FVoiceEndpoint>& Endpoint = ExternalEndpoints.Emplace_GetRef(new FVoiceEndpoint(InDeviceName, UVOIPStatics::GetVoiceSampleRate(), DEFAULT_NUM_VOICE_CHANNELS));
+	TUniquePtr<FVoiceEndpoint>& Endpoint = ExternalEndpoints.Emplace_GetRef(new FVoiceEndpoint(InDeviceName, UVOIPStatics::GetVoiceSampleRate(), UVOIPStatics::GetVoiceNumChannels()));
 	Audio::FPatchOutputStrongPtr OutputPatch = AllRemoteTalkerAudio.AddNewOutput(4096 * 2, 1.0f);
 	Endpoint->PatchInOutput(OutputPatch);
 	return true;
@@ -854,7 +854,8 @@ bool FVoiceEngineImpl::PatchRemoteTalkerOutputToEndpoint(const FString& InDevice
 
 bool FVoiceEngineImpl::PatchLocalTalkerOutputToEndpoint(const FString& InDeviceName)
 {
-	TUniquePtr<FVoiceEndpoint>& Endpoint = ExternalEndpoints.Emplace_GetRef(new FVoiceEndpoint(InDeviceName, UVOIPStatics::GetVoiceSampleRate(), DEFAULT_NUM_VOICE_CHANNELS));
+	// Local talker patched output is always mixed down to mono.
+	TUniquePtr<FVoiceEndpoint>& Endpoint = ExternalEndpoints.Emplace_GetRef(new FVoiceEndpoint(InDeviceName, UVOIPStatics::GetVoiceSampleRate(), 1));
 	Audio::FPatchOutputStrongPtr OutputPatch = VoiceCapture->GetMicrophoneAudio(4096 * 2, 1.0f);
 	Endpoint->PatchInOutput(OutputPatch);
 	return true;
