Change voice engine impl to keep an OSS instance name and not cache the pointer directly

#rnx

diff --git a/Engine/Plugins/Online/OnlineSubsystemUtils/Source/OnlineSubsystemUtils/Private/VoiceEngineImpl.cpp b/Engine/Plugins/Online/OnlineSubsystemUtils/Source/OnlineSubsystemUtils/Private/VoiceEngineImpl.cpp
--- a/Engine/Plugins/Online/OnlineSubsystemUtils/Source/OnlineSubsystemUtils/Private/VoiceEngineImpl.cpp
+++ b/Engine/Plugins/Online/OnlineSubsystemUtils/Source/OnlineSubsystemUtils/Private/VoiceEngineImpl.cpp
@@ -170,7 +170,7 @@ void FRemoteTalkerDataImpl::Cleanup()
 }
 
 FVoiceEngineImpl ::FVoiceEngineImpl()
-	: OnlineSubsystem(nullptr)
+	: OnlineInstanceName(NAME_None),
 	, VoiceCapture(nullptr)
 	, VoiceEncoder(nullptr)
 	, OwningUserIndex(INVALID_INDEX)
@@ -187,7 +187,7 @@ FVoiceEngineImpl ::FVoiceEngineImpl()
 }
 
 FVoiceEngineImpl::FVoiceEngineImpl(IOnlineSubsystem* InSubsystem) :
-	OnlineSubsystem(InSubsystem),
+	OnlineInstanceName(NAME_None),
 	VoiceCapture(nullptr),
 	VoiceEncoder(nullptr),
 	OwningUserIndex(INVALID_INDEX),
@@ -202,6 +202,11 @@ FVoiceEngineImpl::FVoiceEngineImpl(IOnlineSubsystem* InSubsystem) :
 #endif
 {
 	FCoreUObjectDelegates::PostLoadMapWithWorld.AddRaw(this, &FVoiceEngineImpl::OnPostLoadMap);
+
+	if (InSubsystem)
+	{
+		OnlineInstanceName = InSubsystem->GetInstanceName();
+	}
 }
 
 FVoiceEngineImpl::~FVoiceEngineImpl()
@@ -277,7 +282,9 @@ bool FVoiceEngineImpl::Init(int32 MaxLocalTalkers, int32 MaxRemoteTalkers)
 {
 	bool bSuccess = false;
 
-	if (!OnlineSubsystem->IsDedicated())
+	IOnlineSubsystem* OnlineSub = GetOnlineSubSystem();
+
+	if (OnlineSub && !OnlineSub->IsDedicated())
 	{
 		FVoiceModule& VoiceModule = FVoiceModule::Get();
 		if (VoiceModule.IsVoiceEnabled())
@@ -932,6 +939,16 @@ void FVoiceEngineImpl::OnDefaultDeviceChanged()
 }
 #endif
 
+IOnlineSubsystem* FVoiceEngineImpl::GetOnlineSubSystem()
+{
+	if (UWorld* World = GetWorldForOnline(OnlineInstanceName))
+	{
+		return Online::GetSubsystem(World);
+	}
+
+	return nullptr;
+}
+
 FVoiceEndpoint::FVoiceEndpoint(const FString& InEndpointName, float InSampleRate, int32 InNumChannels)
 	: NumChannelsComingIn(InNumChannels)
 {
diff --git a/Engine/Plugins/Online/OnlineSubsystemUtils/Source/OnlineSubsystemUtils/Public/VoiceEngineImpl.h b/Engine/Plugins/Online/OnlineSubsystemUtils/Source/OnlineSubsystemUtils/Public/VoiceEngineImpl.h
--- a/Engine/Plugins/Online/OnlineSubsystemUtils/Source/OnlineSubsystemUtils/Public/VoiceEngineImpl.h
+++ b/Engine/Plugins/Online/OnlineSubsystemUtils/Source/OnlineSubsystemUtils/Public/VoiceEngineImpl.h
@@ -158,8 +158,8 @@ class ONLINESUBSYSTEMUTILS_API FVoiceEngineImpl : public IVoiceEngine, public FS
 	/** Mapping of UniqueIds to the incoming voice data and their audio component */
 	typedef TMap<FUniqueNetIdWrapper, FRemoteTalkerDataImpl> FRemoteTalkerData;
 
-	/** Reference to the main online subsystem */
-	IOnlineSubsystem* OnlineSubsystem;
+	/** Instance name of associated online subsystem */
+	FName OnlineInstanceName;
 
 	FLocalVoiceData PlayerVoiceData[MAX_SPLITSCREEN_TALKERS];
 	/** Reference to voice capture device */
@@ -358,7 +358,7 @@ private:
 	void OnPostLoadMap(UWorld*);
 
 protected:
-	virtual IOnlineSubsystem*				 GetOnlineSubSystem()			{ return OnlineSubsystem; }
+	virtual IOnlineSubsystem*				 GetOnlineSubSystem();
 	virtual const TSharedPtr<IVoiceCapture>& GetVoiceCapture() const		{ return VoiceCapture; }
 	virtual TSharedPtr<IVoiceCapture>&		 GetVoiceCapture()				{ return VoiceCapture; }
 	virtual const TSharedPtr<IVoiceEncoder>& GetVoiceEncoder() const		{ return VoiceEncoder; }
