Fix merge error
#rnx

diff --git a/Engine/Plugins/Online/OnlineSubsystemUtils/Source/OnlineSubsystemUtils/Private/VoiceEngineImpl.cpp b/Engine/Plugins/Online/OnlineSubsystemUtils/Source/OnlineSubsystemUtils/Private/VoiceEngineImpl.cpp
--- a/Engine/Plugins/Online/OnlineSubsystemUtils/Source/OnlineSubsystemUtils/Private/VoiceEngineImpl.cpp
+++ b/Engine/Plugins/Online/OnlineSubsystemUtils/Source/OnlineSubsystemUtils/Private/VoiceEngineImpl.cpp
@@ -170,7 +170,7 @@ void FRemoteTalkerDataImpl::Cleanup()
 }
 
 FVoiceEngineImpl ::FVoiceEngineImpl()
-	: OnlineInstanceName(NAME_None),
+	: OnlineInstanceName(NAME_None)
 	, VoiceCapture(nullptr)
 	, VoiceEncoder(nullptr)
 	, OwningUserIndex(INVALID_INDEX)
