Add Reset Data To VoiceInterfaceImpl

diff --git a/Engine/Plugins/Online/OnlineSubsystemUtils/Source/OnlineSubsystemUtils/Private/VoiceInterfaceImpl.cpp b/Engine/Plugins/Online/OnlineSubsystemUtils/Source/OnlineSubsystemUtils/Private/VoiceInterfaceImpl.cpp
--- a/Engine/Plugins/Online/OnlineSubsystemUtils/Source/OnlineSubsystemUtils/Private/VoiceInterfaceImpl.cpp
+++ b/Engine/Plugins/Online/OnlineSubsystemUtils/Source/OnlineSubsystemUtils/Private/VoiceInterfaceImpl.cpp
@@ -108,12 +108,8 @@ void FOnlineVoiceImpl::ClearVoicePackets()
 {
 	for (uint32 Index = 0; Index < MAX_SPLITSCREEN_TALKERS; Index++)
 	{
-		FVoicePacketImpl& LocalPacket = VoiceData.LocalPackets[Index];
-		if (LocalPacket.Length > 0)
-		{
-			// Mark the local packet as processed
-			LocalPacket.Length = 0;
-		}
+		// Mark the local packet as processed
+		VoiceData.LocalPackets[Index].ResetData();
 	}
 }
 
diff --git a/Engine/Plugins/Online/OnlineSubsystemUtils/Source/OnlineSubsystemUtils/Private/VoicePacketImpl.cpp b/Engine/Plugins/Online/OnlineSubsystemUtils/Source/OnlineSubsystemUtils/Private/VoicePacketImpl.cpp
--- a/Engine/Plugins/Online/OnlineSubsystemUtils/Source/OnlineSubsystemUtils/Private/VoicePacketImpl.cpp
+++ b/Engine/Plugins/Online/OnlineSubsystemUtils/Source/OnlineSubsystemUtils/Private/VoicePacketImpl.cpp
@@ -94,3 +94,8 @@ void FVoicePacketImpl::Serialize(class FArchive& Ar)
 #endif
 	}
 }
+
+void FVoicePacketImpl::ResetData()
+{
+	Length = 0;
+}
diff --git a/Engine/Plugins/Online/OnlineSubsystemUtils/Source/OnlineSubsystemUtils/Public/VoicePacketImpl.h b/Engine/Plugins/Online/OnlineSubsystemUtils/Source/OnlineSubsystemUtils/Public/VoicePacketImpl.h
--- a/Engine/Plugins/Online/OnlineSubsystemUtils/Source/OnlineSubsystemUtils/Public/VoicePacketImpl.h
+++ b/Engine/Plugins/Online/OnlineSubsystemUtils/Source/OnlineSubsystemUtils/Public/VoicePacketImpl.h
@@ -46,6 +46,8 @@ public:
 	 */
 	FVoicePacketImpl(const FVoicePacketImpl& Other);
 
+	virtual void ResetData();
+
 	//~ Begin FVoicePacket interface
 	virtual uint16 GetTotalPacketSize() override;
 	virtual uint16 GetBufferSize() override;
