Stereo Fixes for FVoiceEngineImpl.

diff --git a/Engine/Plugins/Online/OnlineSubsystemUtils/Source/OnlineSubsystemUtils/Private/VoicePacketBuffer.cpp b/Engine/Plugins/Online/OnlineSubsystemUtils/Source/OnlineSubsystemUtils/Private/VoicePacketBuffer.cpp
--- a/Engine/Plugins/Online/OnlineSubsystemUtils/Source/OnlineSubsystemUtils/Private/VoicePacketBuffer.cpp
+++ b/Engine/Plugins/Online/OnlineSubsystemUtils/Source/OnlineSubsystemUtils/Private/VoicePacketBuffer.cpp
@@ -23,7 +23,7 @@ FVoicePacketBuffer::FVoicePacketBuffer(int32 BufferSize, int32 InNumSamplesUntil
 int32 FVoicePacketBuffer::PopAudio(float* DestBuffer, uint32 RequestedSamples)
 {
 #if DEBUG_POPPING
-	UE_LOG(LogTemp, Log, TEXT("Beginning to pop %d samples..."), RequestedSamples);
+	UE_LOG(LogAudio, Log, TEXT("Beginning to pop %d samples..."), RequestedSamples);
 #endif
 
 	uint32 DestBufferSampleIndex = 0;
@@ -32,7 +32,7 @@ int32 FVoicePacketBuffer::PopAudio(float* DestBuffer, uint32 RequestedSamples)
 	while (DestBufferSampleIndex < RequestedSamples)
 	{
 #if DEBUG_POPPING
-		UE_LOG(LogTemp, Log, TEXT("  Loop start, index %d "), DestBufferSampleIndex);
+		UE_LOG(LogAudio, Log, TEXT("  Loop start, index %d "), DestBufferSampleIndex);
 #endif
 		//If we don't have any data, output silence and short circuit.
 		if (ListHead == nullptr)
@@ -42,7 +42,9 @@ int32 FVoicePacketBuffer::PopAudio(float* DestBuffer, uint32 RequestedSamples)
 			DestBufferSampleIndex += RequestedSamples;
 			if (IdleSampleCounter > NumSamplesUntilIdling)
 			{
-				UE_LOG(LogTemp, Warning, TEXT("Voip listener now idle."));
+#if DEBUG_POPPING
+				UE_LOG(LogAudio, Warning, TEXT("Voip listener now idle."));
+#endif
 				bIsIdle = true;
 			}
 
@@ -58,7 +60,7 @@ int32 FVoicePacketBuffer::PopAudio(float* DestBuffer, uint32 RequestedSamples)
 			const int32 SampleOffset = ListHead->StartSample - SampleCounter;
 			const int32 NumSilentSamples = FMath::Min(SampleOffset, (int32) (RequestedSamples - DestBufferSampleIndex));
 #if DEBUG_POPPING
-			UE_LOG(LogTemp, Log, TEXT("    Sample counter %d behind next packet. Injecting %d samples of silence."), SampleOffset, NumSilentSamples);
+			UE_LOG(LogAudio, Log, TEXT("    Sample counter %d behind next packet. Injecting %d samples of silence."), SampleOffset, NumSilentSamples);
 #endif
 			FMemory::Memset(&(DestBuffer[DestBufferSampleIndex]), 0, NumSilentSamples * sizeof(float));
 			DestBufferSampleIndex += NumSilentSamples;
@@ -66,7 +68,7 @@ int32 FVoicePacketBuffer::PopAudio(float* DestBuffer, uint32 RequestedSamples)
 			IdleSampleCounter += NumSilentSamples;
 			if (IdleSampleCounter > NumSamplesUntilIdling)
 			{
-				UE_LOG(LogTemp, Warning, TEXT("Voip listener now idle."));
+				UE_LOG(LogAudio, Warning, TEXT("Voip listener now idle."));
 				bIsIdle = true;
 			}
 #endif
@@ -84,7 +86,7 @@ int32 FVoicePacketBuffer::PopAudio(float* DestBuffer, uint32 RequestedSamples)
 			DestBufferSampleIndex += NumSamplesLeftInPacketBuffer;
 #if DEBUG_POPPING
 			check((ListHead->SamplesLeft - NumSamplesLeftInPacketBuffer) == 0);
-			UE_LOG(LogTemp, Log, TEXT("    Copied %d samples from packet and reached end of packet."), NumSamplesLeftInPacketBuffer);
+			UE_LOG(LogAudio, Log, TEXT("    Copied %d samples from packet and reached end of packet."), NumSamplesLeftInPacketBuffer);
 #endif
 			//Delete this packet and move on to the next one.
 			FreedIndicesQueue.Enqueue(ListHead->IndexInPacketBuffer);
@@ -100,7 +102,7 @@ int32 FVoicePacketBuffer::PopAudio(float* DestBuffer, uint32 RequestedSamples)
 			FMemory::Memcpy(&(DestBuffer[DestBufferSampleIndex]), PacketBufferPtr, SamplesLeft * sizeof(float));
 #if DEBUG_POPPING
 			check(DestBufferSampleIndex + SamplesLeft == RequestedSamples);
-			UE_LOG(LogTemp, Log, TEXT("    Copied %d samples from packet and reached end of buffer."), SamplesLeft);
+			UE_LOG(LogAudio, Log, TEXT("    Copied %d samples from packet and reached end of buffer."), SamplesLeft);
 #endif
 			ListHead->SamplesLeft -= SamplesLeft;
 			SampleCounter += SamplesLeft;
@@ -109,12 +111,12 @@ int32 FVoicePacketBuffer::PopAudio(float* DestBuffer, uint32 RequestedSamples)
 			bIsIdle = false;
 		}
 #if DEBUG_POPPING
-		UE_LOG(LogTemp, Log, TEXT("  Ended loop with %d samples popped to buffer."), DestBufferSampleIndex);
+		UE_LOG(LogAudio, Log, TEXT("  Ended loop with %d samples popped to buffer."), DestBufferSampleIndex);
 #endif
 	}
 #if DEBUG_POPPING
 	check(DestBufferSampleIndex == RequestedSamples);
-	UE_LOG(LogTemp, Log, TEXT("Finished PopAudio."));
+	UE_LOG(LogAudio, Log, TEXT("Finished PopAudio."));
 #endif
 	return RequestedSamples - IdleSampleCounter;
 }
@@ -152,7 +154,7 @@ void FVoicePacketBuffer::PushPacket(const void* InBuffer, int32 NumBytes, uint64
 	}
 	else
 	{
-		UE_LOG(LogTemp, Warning, TEXT("Voice packet buffer filled to capacity of %d packets; packet dropped."), PacketBuffer.Num());
+		UE_LOG(LogAudio, Warning, TEXT("Voice packet buffer filled to capacity of %d packets; packet dropped."), PacketBuffer.Num());
 		return;
 	}
 
@@ -209,13 +211,13 @@ void FVoicePacketBuffer::PushPacket(const void* InBuffer, int32 NumBytes, uint64
 	}
 
 #if DEBUG_SORTING
-	UE_LOG(LogTemp, Log, TEXT("Packet sort order: "));
+	UE_LOG(LogAudio, Log, TEXT("Packet sort order: "));
 	FSortedVoicePacketNode* PacketPtr = ListHead;
 	int32 LoggedPatcketIndex = 0;
 	uint64 CachedSampleStart = 0;
 	while (PacketPtr != nullptr)
 	{
-		UE_LOG(LogTemp, Log, TEXT("    Packet %d: Sample Start %d"), ++LoggedPatcketIndex, PacketPtr->StartSample);
+		UE_LOG(LogAudio, Log, TEXT("    Packet %d: Sample Start %d"), ++LoggedPatcketIndex, PacketPtr->StartSample);
 		check(PacketPtr->StartSample > CachedSampleStart);
 		CachedSampleStart = PacketPtr->StartSample;
 		PacketPtr = PacketPtr->NextPacket;
@@ -223,8 +225,8 @@ void FVoicePacketBuffer::PushPacket(const void* InBuffer, int32 NumBytes, uint64
 	if (LoggedPatcketIndex > 6)
 	{
 		//Put a breakpoint here to check out sample counts.
-		UE_LOG(LogTemp, Log, TEXT("Several Packets buffered."));
+		UE_LOG(LogAudio, Log, TEXT("Several Packets buffered."));
 	}
-	UE_LOG(LogTemp, Log, TEXT("End of packet list."));
+	UE_LOG(LogAudio, Log, TEXT("End of packet list."));
 #endif
 }
\ No newline at end of file
diff --git a/Engine/Source/Runtime/Engine/Private/VoiceConfig.cpp b/Engine/Source/Runtime/Engine/Private/VoiceConfig.cpp
--- a/Engine/Source/Runtime/Engine/Private/VoiceConfig.cpp
+++ b/Engine/Source/Runtime/Engine/Private/VoiceConfig.cpp
@@ -83,16 +83,16 @@ uint32 UVOIPStatics::GetMaxVoiceDataSize()
 	{
 		case 24000:
 		{
-			return 14 * 1024;
+			return 14 * 1024 * GetVoiceNumChannels();
 		}
 		case 48000:
 		{
-			return 32 * 1024;
+			return 32 * 1024 * GetVoiceNumChannels();
 		}
 		default:
 		case 16000:
 		{
-			return 8 * 1024;
+			return 8 * 1024 * GetVoiceNumChannels();
 		}
 	}
 }
@@ -103,7 +103,7 @@ uint32 UVOIPStatics::GetMaxUncompressedVoiceDataSizePerChannel()
 	// This amounts to approximates a second of audio for the minimum voice engine tick frequency.
 	// At 48 kHz, DirectSound will occasionally have to load up to 256 samples into the overflow buffer.
 	// This is the reason for the added 1024 bytes.
-	return GetVoiceSampleRate() * sizeof(uint16) + 1024;
+	return GetVoiceSampleRate() * sizeof(uint16) * GetVoiceNumChannels() + 1024 * GetVoiceNumChannels();
 }
 
 uint32 UVOIPStatics::GetMaxCompressedVoiceDataSize()
@@ -115,16 +115,16 @@ uint32 UVOIPStatics::GetMaxCompressedVoiceDataSize()
 	{
 		case 24000:
 		{
-			return 2 * 1024;
+			return 2 * 1024 * GetVoiceNumChannels();
 		}
 		case 48000:
 		{
-			return 4 * 1024;
+			return 4 * 1024 * GetVoiceNumChannels();
 		}
 		default:
 		case 16000:
 		{
-			return 1 * 1024;
+			return 1 * 1024 * GetVoiceNumChannels();
 		}
 	}
 }
diff --git a/Engine/Source/Runtime/Online/Voice/Private/VoiceDelayBuffer.h b/Engine/Source/Runtime/Online/Voice/Private/VoiceDelayBuffer.h
--- a/Engine/Source/Runtime/Online/Voice/Private/VoiceDelayBuffer.h
+++ b/Engine/Source/Runtime/Online/Voice/Private/VoiceDelayBuffer.h
@@ -61,6 +61,14 @@ public:
 		InternalBuffer[WriteIndex++] = InSample;
 	}
 
+	void PushFrame(const SampleType* InFrame, int32 NumChannels)
+	{
+		for (int32 ChannelIndex = 0; ChannelIndex < NumChannels; ChannelIndex++)
+		{
+			PushSample(InFrame[ChannelIndex]);
+		}
+	}
+
 	/** pop sample off of the buffer. Increments delay between read and write by one sample. */
 	SampleType PopSample()
 	{
@@ -94,8 +102,8 @@ public:
 		{
 			const int32 SamplesUntilEndOfBuffer = InternalBuffer.Num() - ReadIndex;
 			//Copy the rest of the buffer to the end, then the remainder in the beginning.
-			FMemory::Memcpy(OutBuffer, InternalBuffer.GetData() + ReadIndex, SamplesUntilEndOfBuffer);
-			FMemory::Memcpy(OutBuffer + SamplesUntilEndOfBuffer, InternalBuffer.GetData(), SamplesToCopy - SamplesUntilEndOfBuffer);
+			FMemory::Memcpy(OutBuffer, InternalBuffer.GetData() + ReadIndex, SamplesUntilEndOfBuffer * sizeof(SampleType));
+			FMemory::Memcpy(OutBuffer + SamplesUntilEndOfBuffer, InternalBuffer.GetData(), (SamplesToCopy - SamplesUntilEndOfBuffer) * sizeof(SampleType));
 		}
 		ReadIndex = ReadIndex + SamplesToCopy;
 		return SamplesToCopy;
diff --git a/Engine/Source/Runtime/Online/Voice/Private/Windows/VoiceCaptureWindows.cpp b/Engine/Source/Runtime/Online/Voice/Private/Windows/VoiceCaptureWindows.cpp
--- a/Engine/Source/Runtime/Online/Voice/Private/Windows/VoiceCaptureWindows.cpp
+++ b/Engine/Source/Runtime/Online/Voice/Private/Windows/VoiceCaptureWindows.cpp
@@ -444,7 +444,7 @@ void FVoiceCaptureWindows::ProcessData()
 				{
 					if (bMicReleased)
 					{
-						ReleaseBuffer.PushSample(Temp);
+						ReleaseBuffer.PushFrame(&InputBuffer[FrameIndex], NumInputChannels);
 
 						if (!bSampleStartCached)
 						{
@@ -454,8 +454,8 @@ void FVoiceCaptureWindows::ProcessData()
 					}
 					else
 					{
-						AudioBuffer[FrameIndex] = Temp;
-						SamplesPushedToUncompressedAudioBuffer++;
+						FMemory::Memcpy(&AudioBuffer[FrameIndex], &InputBuffer[FrameIndex], sizeof(int16) * NumInputChannels);
+						SamplesPushedToUncompressedAudioBuffer += NumInputChannels;
 					}
 				}
 				else
@@ -488,7 +488,7 @@ void FVoiceCaptureWindows::ProcessData()
 				{
 					if (bMicReleased)
 					{
-						ReleaseBuffer.PushSample(Temp);
+						ReleaseBuffer.PushFrame(&InputBuffer[FrameIndex], NumInputChannels);
 
 						if (!bSampleStartCached)
 						{
@@ -498,8 +498,8 @@ void FVoiceCaptureWindows::ProcessData()
 					}
 					else
 					{
-						AudioBuffer[FrameIndex] = Temp;
-						SamplesPushedToUncompressedAudioBuffer++;
+						FMemory::Memcpy(&AudioBuffer[FrameIndex], &InputBuffer[FrameIndex], sizeof(int16) * NumInputChannels);
+						SamplesPushedToUncompressedAudioBuffer += NumInputChannels;
 					}
 				}
 				else
diff --git a/Engine/Source/Runtime/Windows/XAudio2/Private/XAudio2Support.h b/Engine/Source/Runtime/Windows/XAudio2/Private/XAudio2Support.h
--- a/Engine/Source/Runtime/Windows/XAudio2/Private/XAudio2Support.h
+++ b/Engine/Source/Runtime/Windows/XAudio2/Private/XAudio2Support.h
@@ -73,6 +73,7 @@ public:
 
 	HRESULT STDMETHODCALLTYPE OnDefaultDeviceChanged(EDataFlow flow, ERole role, LPCWSTR pwstrDeviceId) override
 	{
+		FScopeLock ScopeLock(&ListenerArrayMutationLock);
 		for (IDeviceChangedListener* Listener : Listeners)
 		{
 			Listener->OnDefaultDeviceChanged();
@@ -87,6 +88,7 @@ public:
 
 	HRESULT STDMETHODCALLTYPE OnDeviceRemoved(LPCWSTR pwstrDeviceId) override
 	{
+		FScopeLock ScopeLock(&ListenerArrayMutationLock);
 		for (IDeviceChangedListener* Listener : Listeners)
 		{
 			Listener->OnDeviceRemoved(FString(pwstrDeviceId));
@@ -98,6 +100,7 @@ public:
 	{
 		if (dwNewState == DEVICE_STATE_DISABLED || dwNewState == DEVICE_STATE_UNPLUGGED || dwNewState == DEVICE_STATE_NOTPRESENT)
 		{
+			FScopeLock ScopeLock(&ListenerArrayMutationLock);
 			for (IDeviceChangedListener* Listener : Listeners)
 			{
 				Listener->OnDeviceRemoved(FString(pwstrDeviceId));
@@ -133,17 +136,20 @@ public:
 
 	void RegisterDeviceChangedListener(IDeviceChangedListener* DeviceChangedListener)
 	{
+		FScopeLock ScopeLock(&ListenerArrayMutationLock);
 		Listeners.Add(DeviceChangedListener);
 	}
 
 	void UnRegisterDeviceDeviceChangedListener(IDeviceChangedListener* DeviceChangedListener)
 	{
+		FScopeLock ScopeLock(&ListenerArrayMutationLock);
 		Listeners.Remove(DeviceChangedListener);
 	}
 
 private:
 	LONG Ref;
 	TSet<IDeviceChangedListener*> Listeners;
+	FCriticalSection ListenerArrayMutationLock;
 	IMMDeviceEnumerator* DeviceEnumerator;
 	bool bComInitialized;
 };
