Change OnlineSubsystemUtiles CreateVoiceSynthComponent to take World as a param to for voice in PIE

diff --git a/Engine/Plugins/Online/OnlineSubsystemUtils/Source/OnlineSubsystemUtils/Private/OnlineSubsystemUtils.cpp b/Engine/Plugins/Online/OnlineSubsystemUtils/Source/OnlineSubsystemUtils/Private/OnlineSubsystemUtils.cpp
--- a/Engine/Plugins/Online/OnlineSubsystemUtils/Source/OnlineSubsystemUtils/Private/OnlineSubsystemUtils.cpp
+++ b/Engine/Plugins/Online/OnlineSubsystemUtils/Source/OnlineSubsystemUtils/Private/OnlineSubsystemUtils.cpp
@@ -184,6 +184,42 @@ UVoipListenerSynthComponent* CreateVoiceSynthComponent(uint32 SampleRate)
 	return SynthComponentPtr;
 }
 
+
+UVoipListenerSynthComponent* CreateVoiceSynthComponent(UWorld* World, uint32 SampleRate)
+{
+	UVoipListenerSynthComponent* SynthComponentPtr = nullptr;
+
+	if (World)
+	{
+		if (FAudioDevice* AudioDevice = World->GetAudioDevice())
+		{
+			SynthComponentPtr = NewObject<UVoipListenerSynthComponent>();
+			if (SynthComponentPtr)
+			{
+				const FSoftObjectPath VoiPSoundClassName = GetDefault<UAudioSettings>()->VoiPSoundClass;
+				if (VoiPSoundClassName.IsValid())
+				{
+					SynthComponentPtr->SoundClass = LoadObject<USoundClass>(nullptr, *VoiPSoundClassName.ToString());
+				}
+
+				if (CvarAlwaysPlayVoipComponent)
+				{
+					SynthComponentPtr->bAlwaysPlay = true;
+				}
+
+				SynthComponentPtr->RegisterComponentWithWorld(World);
+				SynthComponentPtr->Initialize(SampleRate);
+			}
+			else
+			{
+				UE_LOG(LogVoiceDecode, Warning, TEXT("Unable to create voice synth component!"));
+			}
+		}
+	}
+
+	return SynthComponentPtr;
+}
+
 void ApplyVoiceSettings(UVoipListenerSynthComponent* InSynthComponent, const FVoiceSettings& InSettings)
 {
 	InSynthComponent->CreateAudioComponent();
diff --git a/Engine/Plugins/Online/OnlineSubsystemUtils/Source/OnlineSubsystemUtils/Private/VoiceEngineImpl.cpp b/Engine/Plugins/Online/OnlineSubsystemUtils/Source/OnlineSubsystemUtils/Private/VoiceEngineImpl.cpp
--- a/Engine/Plugins/Online/OnlineSubsystemUtils/Source/OnlineSubsystemUtils/Private/VoiceEngineImpl.cpp
+++ b/Engine/Plugins/Online/OnlineSubsystemUtils/Source/OnlineSubsystemUtils/Private/VoiceEngineImpl.cpp
@@ -587,7 +587,14 @@ uint32 FVoiceEngineImpl::SubmitRemoteVoiceData(const FUniqueNetIdWrapper& Remote
 	{
 		CreateSerializeHelper();
 
-		QueuedData.VoipSynthComponent = CreateVoiceSynthComponent(UVOIPStatics::GetVoiceSampleRate());
+		if (GetOnlineSubSystem())
+		{
+			if (UWorld* World = GetWorldForOnline(GetOnlineSubSystem()->GetInstanceName()))
+			{
+				QueuedData.VoipSynthComponent = CreateVoiceSynthComponent(World, UVOIPStatics::GetVoiceSampleRate());
+			}
+		}
+
 		if (QueuedData.VoipSynthComponent)
 		{
 			//TODO, make buffer size and buffering delay runtime-controllable parameters.
diff --git a/Engine/Plugins/Online/OnlineSubsystemUtils/Source/OnlineSubsystemUtils/Public/OnlineSubsystemUtils.h b/Engine/Plugins/Online/OnlineSubsystemUtils/Source/OnlineSubsystemUtils/Public/OnlineSubsystemUtils.h
--- a/Engine/Plugins/Online/OnlineSubsystemUtils/Source/OnlineSubsystemUtils/Public/OnlineSubsystemUtils.h
+++ b/Engine/Plugins/Online/OnlineSubsystemUtils/Source/OnlineSubsystemUtils/Public/OnlineSubsystemUtils.h
@@ -22,8 +22,12 @@ class UAudioComponent;
 ONLINESUBSYSTEMUTILS_API UAudioComponent* CreateVoiceAudioComponent(uint32 SampleRate, int32 NumChannels);
 
 /** @return an initialized Synth component specifically for use with VoIP */
+UE_DEPRECATED(4.25, "CreateVoiceSynthComponent(uint32 SampleRate) is deprecated, use CreateVoiceSynthComponent(UWorld* World, uint32 SampleRate).")
 ONLINESUBSYSTEMUTILS_API UVoipListenerSynthComponent* CreateVoiceSynthComponent(uint32 SampleRate);
 
+/** @return an initialized Synth component specifically for use with VoIP */
+ONLINESUBSYSTEMUTILS_API UVoipListenerSynthComponent* CreateVoiceSynthComponent(UWorld* World, uint32 SampleRate);
+
 /** Updates InSynthComponent based on InSettings. */
 ONLINESUBSYSTEMUTILS_API void ApplyVoiceSettings(UVoipListenerSynthComponent* InSynthComponent, const FVoiceSettings& InSettings);
 
