Fix for crash in device changed listener

diff --git a/Engine/Plugins/Online/OnlineSubsystemUtils/Source/OnlineSubsystemUtils/Private/VoiceEngineImpl.cpp b/Engine/Plugins/Online/OnlineSubsystemUtils/Source/OnlineSubsystemUtils/Private/VoiceEngineImpl.cpp
--- a/Engine/Plugins/Online/OnlineSubsystemUtils/Source/OnlineSubsystemUtils/Private/VoiceEngineImpl.cpp
+++ b/Engine/Plugins/Online/OnlineSubsystemUtils/Source/OnlineSubsystemUtils/Private/VoiceEngineImpl.cpp
@@ -188,6 +188,7 @@ FVoiceEngineImpl ::FVoiceEngineImpl()
 	, SerializeHelper(nullptr)
 #if PLATFORM_WINDOWS
 	, bAudioDeviceChanged(false)
+	, bDeviceChangeListenerRegistered(false)
 #endif
 {
 }
@@ -228,6 +229,13 @@ FVoiceEngineImpl::~FVoiceEngineImpl()
 	VoiceEncoder = nullptr;
 
 	delete SerializeHelper;
+
+#if PLATFORM_WINDOWS
+	if (bDeviceChangeListenerRegistered)
+	{
+		UnregisterDeviceChangedListener();
+	}
+#endif
 }
 
 void FVoiceEngineImpl::VoiceCaptureUpdate() const
@@ -955,6 +963,8 @@ void FVoiceEngineImpl::RegisterDeviceChangedListener()
 	}
 
 	NotificationClient::WindowsNotificationClient->RegisterDeviceChangedListener(this);
+
+	bDeviceChangeListenerRegistered = true;
 }
 
 void FVoiceEngineImpl::UnregisterDeviceChangedListener()
@@ -963,6 +973,8 @@ void FVoiceEngineImpl::UnregisterDeviceChangedListener()
 	{
 		NotificationClient::WindowsNotificationClient->UnRegisterDeviceDeviceChangedListener(this);
 	}
+
+	bDeviceChangeListenerRegistered = false;
 }
 
 void FVoiceEngineImpl::HandleDeviceChange()
diff --git a/Engine/Plugins/Online/OnlineSubsystemUtils/Source/OnlineSubsystemUtils/Public/VoiceEngineImpl.h b/Engine/Plugins/Online/OnlineSubsystemUtils/Source/OnlineSubsystemUtils/Public/VoiceEngineImpl.h
--- a/Engine/Plugins/Online/OnlineSubsystemUtils/Source/OnlineSubsystemUtils/Public/VoiceEngineImpl.h
+++ b/Engine/Plugins/Online/OnlineSubsystemUtils/Source/OnlineSubsystemUtils/Public/VoiceEngineImpl.h
@@ -385,6 +385,8 @@ protected:
 
 	//~ Begin IDeviceChangedListener
 	virtual void OnDefaultDeviceChanged() override;
+
+	bool bDeviceChangeListenerRegistered;
 #else
 	virtual void OnDefaultDeviceChanged() override {}
 #endif
