Add Cvars for adjusting buffer delay on VOIP synth component.

diff --git a/Engine/Plugins/Online/OnlineSubsystemUtils/Source/OnlineSubsystemUtils/Private/VoicePacketBuffer.cpp b/Engine/Plugins/Online/OnlineSubsystemUtils/Source/OnlineSubsystemUtils/Private/VoicePacketBuffer.cpp
--- a/Engine/Plugins/Online/OnlineSubsystemUtils/Source/OnlineSubsystemUtils/Private/VoicePacketBuffer.cpp
+++ b/Engine/Plugins/Online/OnlineSubsystemUtils/Source/OnlineSubsystemUtils/Private/VoicePacketBuffer.cpp
@@ -26,9 +26,12 @@ int32 FVoicePacketBuffer::PopAudio(float* DestBuffer, uint32 RequestedSamples)
 	UE_LOG(LogAudio, Log, TEXT("Beginning to pop %d samples..."), RequestedSamples);
 #endif
 
+#if SCOPELOCK_VOICE_PACKET_BUFFER
+	FScopeLock ScopeLock(&ListHeadCriticalSection);
+#endif
+
 	uint32 DestBufferSampleIndex = 0;
 	// start churning through packets.
-	// TODO: Worry about threading issues here.
 	while (DestBufferSampleIndex < RequestedSamples)
 	{
 #if DEBUG_POPPING
@@ -118,7 +121,102 @@ int32 FVoicePacketBuffer::PopAudio(float* DestBuffer, uint32 RequestedSamples)
 	check(DestBufferSampleIndex == RequestedSamples);
 	UE_LOG(LogAudio, Log, TEXT("Finished PopAudio."));
 #endif
-	return RequestedSamples - IdleSampleCounter;
+
+	const int32 SamplesPopped = RequestedSamples - IdleSampleCounter;
+	NumBufferedSamples.Subtract(SamplesPopped);
+
+	return SamplesPopped;
+}
+
+int32 FVoicePacketBuffer::GetNumBufferedSamples()
+{
+	return NumBufferedSamples.GetValue();
+}
+
+uint32 FVoicePacketBuffer::DropOldestAudio(uint32 InNumSamples)
+{
+#if DEBUG_POPPING
+	UE_LOG(LogAudio, Log, TEXT("Beginning to seek forward %d samples..."), InNumSamples);
+#endif
+
+#if SCOPELOCK_VOICE_PACKET_BUFFER
+	FScopeLock ScopeLock(&ListHeadCriticalSection);
+#endif
+
+	uint32 NumSamplesSeeked = 0;
+
+	// start churning through packets.
+	while (NumSamplesSeeked < InNumSamples)
+	{
+#if DEBUG_POPPING
+		UE_LOG(LogAudio, Log, TEXT("  Loop start, index %d "), NumSamplesSeeked);
+#endif
+
+		if (ListHead == nullptr)
+		{
+			return NumSamplesSeeked;
+		}
+
+		// If our current packet is more than 0 samples later than our sample counter,
+		// fill silence until we get to the sample index.
+		else if (ListHead->StartSample > SampleCounter)
+		{
+#if DEBUG_SKIPSILENCE
+			SampleCounter = ListHead->StartSample;
+#else
+			
+#if DEBUG_POPPING
+			const int32 SampleOffset = ListHead->StartSample - SampleCounter;
+			const int32 NumSilentSamples = FMath::Min(SampleOffset, (int32)(RequestedSamples - DestBufferSampleIndex));
+			UE_LOG(LogAudio, Log, TEXT("    Sample counter %d behind next packet. Injecting %d samples of silence."), SampleOffset, NumSilentSamples);
+#endif
+#endif
+		}
+		// If we have less samples in the current packet than we need to churn through,
+		// copy over the rest of the samples in this packet and move on to the next packet.
+		else if (ListHead->SamplesLeft <= (int32)(InNumSamples - NumSamplesSeeked))
+		{
+			bIsIdle = false;
+			const uint32 NumSamplesLeftInPacketBuffer = ListHead->SamplesLeft;
+			SampleCounter += NumSamplesLeftInPacketBuffer;
+			NumSamplesSeeked += NumSamplesLeftInPacketBuffer;
+#if DEBUG_POPPING
+			check((ListHead->SamplesLeft - NumSamplesLeftInPacketBuffer) == 0);
+			UE_LOG(LogAudio, Log, TEXT("    Copied %d samples from packet and reached end of packet."), NumSamplesLeftInPacketBuffer);
+#endif
+			//Delete this packet and move on to the next one.
+			FreedIndicesQueue.Enqueue(ListHead->IndexInPacketBuffer);
+			ListHead = ListHead->NextPacket;
+		}
+		// If we have more samples in the current packet than we need to churn through,
+		// copy over samples and update the state in the packet node.
+		else
+		{
+			int32 SamplesLeft = InNumSamples - NumSamplesSeeked;
+#if DEBUG_POPPING
+			check(NumSamplesSeeked + SamplesLeft == InNumSamples);
+			UE_LOG(LogAudio, Log, TEXT("    Copied %d samples from packet and reached end of buffer."), SamplesLeft);
+#endif
+			ListHead->SamplesLeft -= SamplesLeft;
+			SampleCounter += SamplesLeft; 
+			NumSamplesSeeked += SamplesLeft;
+			IdleSampleCounter = 0;
+			bIsIdle = false;
+		}
+#if DEBUG_POPPING
+		UE_LOG(LogAudio, Log, TEXT("  Ended loop with %d samples popped to buffer."), DestBufferSampleIndex);
+#endif
+	}
+
+#if DEBUG_POPPING
+	check(NumSamplesSeeked == InNumSamples);
+	UE_LOG(LogAudio, Log, TEXT("Finished PopAudio."));
+#endif
+
+	const int32 SamplesPopped = InNumSamples - IdleSampleCounter;
+	NumBufferedSamples.Subtract(SamplesPopped);
+
+	return SamplesPopped;
 }
 
 void FVoicePacketBuffer::Reset(int32 InStartSample)
@@ -138,7 +236,6 @@ void FVoicePacketBuffer::Reset(int32 InStartSample)
 
 void FVoicePacketBuffer::PushPacket(const void* InBuffer, int32 NumBytes, uint64 InStartSample, EVoipStreamDataFormat Format)
 {
-	//To do- trash buffers that are too far behind our start sample.
 	if (SampleCounter > InStartSample + NumBytes)
 	{
 		SampleCounter = InStartSample;
@@ -165,10 +262,13 @@ void FVoicePacketBuffer::PushPacket(const void* InBuffer, int32 NumBytes, uint64
 	PacketBuffer[Index].IndexInPacketBuffer = Index;
 	FSortedVoicePacketNode* PacketNodePtr = &(PacketBuffer[Index]);
 
+#if SCOPELOCK_VOICE_PACKET_BUFFER
+	FScopeLock ScopeLock(&ListHeadCriticalSection);
+#endif
+
 	// If our buffer is empty, we can just place this packet at the beginning.
 	// Otherwise, we need to traverse our sorted list and insert the packet at
 	// the correct place based on StartSample.
-	// TODO: Worry about threading issues here.
 	if (ListHead == nullptr)
 	{
 		ListHead = PacketNodePtr;
@@ -210,6 +310,8 @@ void FVoicePacketBuffer::PushPacket(const void* InBuffer, int32 NumBytes, uint64
 		CurrentPacket->NextPacket = PacketNodePtr;
 	}
 
+	NumBufferedSamples.Add(NumBytes / SizeOfSample(Format));
+
 #if DEBUG_SORTING
 	UE_LOG(LogAudio, Log, TEXT("Packet sort order: "));
 	FSortedVoicePacketNode* PacketPtr = ListHead;
diff --git a/Engine/Plugins/Online/OnlineSubsystemUtils/Source/OnlineSubsystemUtils/Private/VoipListenerSynthComponent.cpp b/Engine/Plugins/Online/OnlineSubsystemUtils/Source/OnlineSubsystemUtils/Private/VoipListenerSynthComponent.cpp
--- a/Engine/Plugins/Online/OnlineSubsystemUtils/Source/OnlineSubsystemUtils/Private/VoipListenerSynthComponent.cpp
+++ b/Engine/Plugins/Online/OnlineSubsystemUtils/Source/OnlineSubsystemUtils/Private/VoipListenerSynthComponent.cpp
@@ -24,6 +24,13 @@ FAutoConsoleVariableRef CVarDefaultPatchGain(
 	TEXT("Changes the default gain of audio patches, in linear gain.\n"),
 	ECVF_Default);
 
+static int32 ShouldResyncCVar = 1;
+FAutoConsoleVariableRef CVarShouldResync(
+	TEXT("voice.playback.ShouldResync"),
+	ShouldResyncCVar,
+	TEXT("If set to 1, we will resync VOIP audio once it's latency goes beyond voice.playback.ResyncThreshold.\n"),
+	ECVF_Default);
+
 static int32 MuteAudioEngineOutputCVar = 0;
 FAutoConsoleVariableRef CVarMuteAudioEngineOutput(
 	TEXT("voice.MuteAudioEngineOutput"),
@@ -31,6 +38,14 @@ FAutoConsoleVariableRef CVarMuteAudioEngineOutput(
 	TEXT("When set to a nonzero value, the output for the audio engine will be muted..\n"),
 	ECVF_Default);
 
+static float ResyncThresholdCVar = 0.3f;
+FAutoConsoleVariableRef CVarResyncThreshold(
+	TEXT("voice.playback.ResyncThreshold"),
+	ResyncThresholdCVar,
+	TEXT("If the amount of audio we have buffered is greater than this value, we drop the oldest audio we have and sync to have voice.JitterDelay worth of buffered audio.\n"),
+	ECVF_Default);
+
+
 
 bool UVoipListenerSynthComponent::Init(int32& SampleRate)
 {
@@ -54,6 +69,12 @@ int32 UVoipListenerSynthComponent::OnGenerateAudio(float* OutAudio, int32 NumSam
 	}
 	else if (PacketBuffer.IsValid())
 	{
+		// Handle resync, if neccessary.
+		if (ShouldResyncCVar)
+		{
+			ForceResync();
+		}
+
 		PacketBuffer->PopAudio(OutAudio, NumSamples);
 	}
 
@@ -88,6 +109,7 @@ void UVoipListenerSynthComponent::OpenPacketStream(uint64 BeginningSampleCount,
 	ClosePacketStream();
 	PacketBuffer.Reset(new FVoicePacketBuffer(BufferSize, NumSecondsUntilIdling * MySampleRate, BeginningSampleCount));
 	
+	JitterDelayInSeconds = JitterDelay;
 	PreDelaySampleCounter = JitterDelay * MySampleRate;
 }
 
@@ -155,6 +177,22 @@ uint64 UVoipListenerSynthComponent::GetSampleCounter()
 }
 
 
+void UVoipListenerSynthComponent::ForceResync()
+{
+	check(PacketBuffer.IsValid());
+
+	const int32 TargetLatencyInSamples = FMath::FloorToInt(JitterDelayInSeconds * NumChannels * MySampleRate);
+	const int32 CurrentLatencyInSamples = PacketBuffer->GetNumBufferedSamples();
+	const int32 ResyncThresholdInSamples = FMath::FloorToInt(ResyncThresholdCVar * NumChannels * MySampleRate);
+
+	int32 AmountToSkip = CurrentLatencyInSamples - TargetLatencyInSamples;
+
+	if (AmountToSkip > ResyncThresholdInSamples)
+	{
+		PacketBuffer->DropOldestAudio(AmountToSkip);
+	}
+}
+
 #if DEBUG_BUFFERING
 FDebugFMTone::FDebugFMTone(float InSampleRate, float InCarrierFreq, float InModFreq, float InCarrierAmp, float InModAmp)
 	: SampleRate(InSampleRate)
diff --git a/Engine/Plugins/Online/OnlineSubsystemUtils/Source/OnlineSubsystemUtils/Public/VoicePacketBuffer.h b/Engine/Plugins/Online/OnlineSubsystemUtils/Source/OnlineSubsystemUtils/Public/VoicePacketBuffer.h
--- a/Engine/Plugins/Online/OnlineSubsystemUtils/Source/OnlineSubsystemUtils/Public/VoicePacketBuffer.h
+++ b/Engine/Plugins/Online/OnlineSubsystemUtils/Source/OnlineSubsystemUtils/Public/VoicePacketBuffer.h
@@ -7,6 +7,9 @@
 #include "AudioMixerTypes.h"
 #include "Net/VoiceDataCommon.h"
 
+// Set this to 1 to enforce a critical section between FVoicePacketBuffer::PopAudio and FVoicePacketBuffer::SubmitPacket.
+#define SCOPELOCK_VOICE_PACKET_BUFFER 0
+
 enum class EVoipStreamDataFormat : uint8
 {
 	Float,
@@ -138,10 +141,38 @@ public:
 	// Get whatever the most recent sample processed is for buffer-accurate timing.
 	uint64 GetCurrentSample() { return SampleCounter; }
 
+	// This will return the current amount of samples buffered here. 
+	int32 GetNumBufferedSamples();
+
+	// When called, will seek forward in the buffer by the requested number of samples, skipping and removing the audio we seek past.
+	// This will return the number of samples we were able to seek forward by.
+	uint32 DropOldestAudio(uint32 InNumSamples);
+
 	// Clears the buffer.
 	void Reset(int32 InStartSample);
 
 private:
+
+	static int32 SizeOfSample(EVoipStreamDataFormat InFormat)
+	{
+		switch (InFormat)
+		{
+			case EVoipStreamDataFormat::Float:
+			{
+				return sizeof(float);
+			}
+			case EVoipStreamDataFormat::Int16:
+			{
+				return sizeof(int16);
+			}
+			default:
+			{
+				checkNoEntry();
+				return 1;
+			}
+		}
+	}
+
 	// Default constructor. Hidden on purpose.
 	FVoicePacketBuffer();
 
@@ -165,4 +196,11 @@ private:
 
 	// After a packet is used, we push its index on the PacketBuffer so that we can Initialize it.
 	TQueue<int32> FreedIndicesQueue;
+
+	// The amount of samples of audio we currently have buffered.
+	FThreadSafeCounter NumBufferedSamples;
+
+#if SCOPELOCK_VOICE_PACKET_BUFFER
+	FCriticalSection ListHeadCriticalSection;
+#endif
 };
diff --git a/Engine/Plugins/Online/OnlineSubsystemUtils/Source/OnlineSubsystemUtils/Public/VoipListenerSynthComponent.h b/Engine/Plugins/Online/OnlineSubsystemUtils/Source/OnlineSubsystemUtils/Public/VoipListenerSynthComponent.h
--- a/Engine/Plugins/Online/OnlineSubsystemUtils/Source/OnlineSubsystemUtils/Public/VoipListenerSynthComponent.h
+++ b/Engine/Plugins/Online/OnlineSubsystemUtils/Source/OnlineSubsystemUtils/Public/VoipListenerSynthComponent.h
@@ -101,16 +101,25 @@ public:
 	uint64 GetSampleCounter();
 
 private:
+	/**
+	 * This call will check how much audio we have buffered and drop the oldest audio if necessary.
+	 */
+	void ForceResync();
+
+
 	/** This patch input can optionally be set  */
 	Audio::FPatchInput ExternalSend;
 
-	// This is allocated on OpenPacketStream()
+	/** This is allocated on OpenPacketStream() */
 	TUniquePtr<FVoicePacketBuffer> PacketBuffer;
 
 	float MySampleRate;
 
 	int32 PreDelaySampleCounter;
 
+	/** The amount of audio we keep buffered to reduce underruns. */
+	float JitterDelayInSeconds;
+
 #if DEBUG_BUFFERING
 	FDebugFMTone FMToneGenerator;
 #endif
diff --git a/Engine/Source/Runtime/Engine/Classes/Sound/AudioSettings.h b/Engine/Source/Runtime/Engine/Classes/Sound/AudioSettings.h
--- a/Engine/Source/Runtime/Engine/Classes/Sound/AudioSettings.h
+++ b/Engine/Source/Runtime/Engine/Classes/Sound/AudioSettings.h
@@ -106,10 +106,6 @@ class ENGINE_API UAudioSettings : public UDeveloperSettings
 	UPROPERTY(config, EditAnywhere, Category = "Audio", meta = (DisplayName = "VOIP Sample Rate"))
 	EVoiceSampleRate VoiPSampleRate;
 
-	/** The amount of time to buffer incoming voice audio for ahead of time. Increasing this value can help avoid jitter on slower network connections. */
-	UPROPERTY(config, EditAnywhere, Category = "Audio", AdvancedDisplay, meta = (ClampMin = 0.05, ClampMax = 3.0))
-	float VoipBufferingDelay;
-
 	/** The amount of audio to send to reverb submixes if no reverb send is setup for the source through attenuation settings. Only used in audio mixer. */
 	UPROPERTY(config)
 	float DefaultReverbSendLevel_DEPRECATED;
diff --git a/Engine/Source/Runtime/Engine/Private/AudioSettings.cpp b/Engine/Source/Runtime/Engine/Private/AudioSettings.cpp
--- a/Engine/Source/Runtime/Engine/Private/AudioSettings.cpp
+++ b/Engine/Source/Runtime/Engine/Private/AudioSettings.cpp
@@ -36,7 +36,6 @@ void UAudioSettings::AddDefaultSettings()
 	DefaultReverbSendLevel_DEPRECATED = 0.0f;
 	bEnableLegacyReverb = false;
 	VoiPSampleRate = EVoiceSampleRate::Low16000Hz;
-	VoipBufferingDelay = 0.2f;
 	NumStoppingSources = 8;
 }
 
diff --git a/Engine/Source/Runtime/Engine/Private/VoiceConfig.cpp b/Engine/Source/Runtime/Engine/Private/VoiceConfig.cpp
--- a/Engine/Source/Runtime/Engine/Private/VoiceConfig.cpp
+++ b/Engine/Source/Runtime/Engine/Private/VoiceConfig.cpp
@@ -21,6 +21,14 @@ static TAutoConsoleVariable<float> CVarVoiceSilenceDetectionThreshold(TEXT("voic
 	TEXT("Threshold to be set for the VOIP microphone's silence detection algorithm.\n"),	
 	ECVF_Default);
 
+static float JitterBufferDelayCvar = 0.3;
+FAutoConsoleVariableRef CVarJitterBufferDelay(
+	TEXT("voice.JitterBufferDelay"),
+	JitterBufferDelayCvar,
+	TEXT("The default amount of audio we buffer, in seconds, before we play back audio. Decreasing this value will decrease latency but increase the potential for underruns.\n")
+	TEXT("Value: Number of seconds of audio we buffer."),
+	ECVF_Default);
+
 static int32 NumVoiceChannelsCvar = 1;
 FAutoConsoleVariableRef CVarNumVoiceChannels(
 	TEXT("voice.NumChannels"),
@@ -142,18 +150,7 @@ EAudioEncodeHint UVOIPStatics::GetAudioEncodingHint()
 
 float UVOIPStatics::GetBufferingDelay()
 {
-	static bool bRetrievedBufferingDelay = false;
-	static float DesiredBufferDelay = 0.2f;
-	if (!bRetrievedBufferingDelay)
-	{
-		bRetrievedBufferingDelay = true;
-		float SettingsBufferDelay;
-		if (GConfig->GetFloat(TEXT("/Script/Engine.AudioSettings"), TEXT("VoipBufferingDelay"), SettingsBufferDelay, GEngineIni) && SettingsBufferDelay > 0.0f)
-		{
-			DesiredBufferDelay = SettingsBufferDelay;
-		}
-	}
-	return DesiredBufferDelay;
+	return JitterBufferDelayCvar;
 }
 
 int32 UVOIPStatics::GetNumBufferedPackets()
