Fix infinite recursion in UIpConnection::HandleSocketSendResult when SocketSend failure occurs with a PendingNetGame

[RN] minor engine/network

diff --git a/Engine/Plugins/Online/OnlineSubsystemUtils/Source/OnlineSubsystemUtils/Private/IpConnection.cpp b/Engine/Plugins/Online/OnlineSubsystemUtils/Source/OnlineSubsystemUtils/Private/IpConnection.cpp
--- a/Engine/Plugins/Online/OnlineSubsystemUtils/Source/OnlineSubsystemUtils/Private/IpConnection.cpp
+++ b/Engine/Plugins/Online/OnlineSubsystemUtils/Source/OnlineSubsystemUtils/Private/IpConnection.cpp
@@ -383,16 +383,20 @@ void UIpConnection::HandleSocketSendResult(const FSocketSendResult& Result, ISoc
 			}
 		}
 
-		FString ErrorString = FString::Printf(TEXT("UIpNetConnection::HandleSocketSendResult: Socket->SendTo failed with error %i (%s). %s Connection will be closed during next Tick()!"),
-			static_cast<int32>(Result.Error),
-			SocketSubsystem->GetSocketError(Result.Error),
-			*Describe());
+		// Broadcast the error only on the first occurrence
+		if( GetPendingCloseDueToSocketSendFailure() == false )
+		{
+			// Request the connection to be disconnected during next tick() since we got a critical socket failure, the actual disconnect is postponed 		
+			// to avoid issues with the call Close() causing issues with reentrant code paths in DataChannel::SendBunch() and FlushNet()
+			SetPendingCloseDueToSocketSendFailure();
 
-		GEngine->BroadcastNetworkFailure(Driver->GetWorld(), Driver, ENetworkFailure::ConnectionLost, ErrorString);
+			FString ErrorString = FString::Printf(TEXT("UIpNetConnection::HandleSocketSendResult: Socket->SendTo failed with error %i (%s). %s Connection will be closed during next Tick()!"),
+												  static_cast<int32>(Result.Error),
+												  SocketSubsystem ? SocketSubsystem->GetSocketError(Result.Error) : TEXT("None"),
+												  *Describe());
 
-		// Request the connection to be disconnected during next tick() since we got a critical socket failure, the actual disconnect is postponed 		
-		// to avoid issues with the call Close() causing issues with reentrant code paths in DataChannel::SendBunch() and FlushNet()
-		SetPendingCloseDueToSocketSendFailure();
+			GEngine->BroadcastNetworkFailure(Driver->GetWorld(), Driver, ENetworkFailure::ConnectionLost, ErrorString);
+		}
 	}
 }
 
