Fix static analysis warning due to extra null ptr check

#rnx

diff --git a/Engine/Plugins/Online/OnlineSubsystemUtils/Source/OnlineSubsystemUtils/Private/IpConnection.cpp b/Engine/Plugins/Online/OnlineSubsystemUtils/Source/OnlineSubsystemUtils/Private/IpConnection.cpp
--- a/Engine/Plugins/Online/OnlineSubsystemUtils/Source/OnlineSubsystemUtils/Private/IpConnection.cpp
+++ b/Engine/Plugins/Online/OnlineSubsystemUtils/Source/OnlineSubsystemUtils/Private/IpConnection.cpp
@@ -392,7 +392,7 @@ void UIpConnection::HandleSocketSendResult(const FSocketSendResult& Result, ISoc
 
 			FString ErrorString = FString::Printf(TEXT("UIpNetConnection::HandleSocketSendResult: Socket->SendTo failed with error %i (%s). %s Connection will be closed during next Tick()!"),
 												  static_cast<int32>(Result.Error),
-												  SocketSubsystem ? SocketSubsystem->GetSocketError(Result.Error) : TEXT("None"),
+												  SocketSubsystem->GetSocketError(Result.Error),
 												  *Describe());
 
 			GEngine->BroadcastNetworkFailure(Driver->GetWorld(), Driver, ENetworkFailure::ConnectionLost, ErrorString);
