Call Reset() from FRemoteTalkerDataImpl's destructor to make sure everything is properly cleaned up, including unregistering the VoipSynthComponent. Prevents VoipSynthComponents from leaking if Reset() isn't called elsewhere.

diff --git a/Engine/Plugins/Online/OnlineSubsystemUtils/Source/OnlineSubsystemUtils/Private/VoiceEngineImpl.cpp b/Engine/Plugins/Online/OnlineSubsystemUtils/Source/OnlineSubsystemUtils/Private/VoiceEngineImpl.cpp
--- a/Engine/Plugins/Online/OnlineSubsystemUtils/Source/OnlineSubsystemUtils/Private/VoiceEngineImpl.cpp
+++ b/Engine/Plugins/Online/OnlineSubsystemUtils/Source/OnlineSubsystemUtils/Private/VoiceEngineImpl.cpp
@@ -120,12 +120,7 @@ FRemoteTalkerDataImpl::~FRemoteTalkerDataImpl()
 {
 	VoiceDecoder = nullptr;
 
-	CurrentUncompressedDataQueueSize = 0;
-
-	{
-		FScopeLock ScopeLock(&QueueLock);
-		UncompressedDataQueue.Empty();
-	}
+	Reset();
 }
 
 void FRemoteTalkerDataImpl::Reset()
@@ -139,7 +134,7 @@ void FRemoteTalkerDataImpl::Reset()
 		VoipSynthComponent->Stop();
 
 		UAudioComponent* AudioComponent = VoipSynthComponent->GetAudioComponent();
-		if (AudioComponent->IsRegistered())
+		if (AudioComponent && AudioComponent->IsRegistered())
 		{
 			AudioComponent->UnregisterComponent();
 		}
