Don't allocate Voice Packet Buffer backing audio memory ahead of time.

diff --git a/Engine/Plugins/Online/OnlineSubsystemUtils/Source/OnlineSubsystemUtils/Private/VoicePacketBuffer.cpp b/Engine/Plugins/Online/OnlineSubsystemUtils/Source/OnlineSubsystemUtils/Private/VoicePacketBuffer.cpp
--- a/Engine/Plugins/Online/OnlineSubsystemUtils/Source/OnlineSubsystemUtils/Private/VoicePacketBuffer.cpp
+++ b/Engine/Plugins/Online/OnlineSubsystemUtils/Source/OnlineSubsystemUtils/Private/VoicePacketBuffer.cpp
@@ -83,7 +83,8 @@ int32 FVoicePacketBuffer::PopAudio(float* DestBuffer, uint32 RequestedSamples)
 			bIsIdle = false;
 			const uint32 NumSamplesLeftInPacketBuffer = ListHead->SamplesLeft;
 			const int32 PacketSampleIndex = ListHead->BufferNumSamples - NumSamplesLeftInPacketBuffer;
-			const float* PacketBufferPtr = &(ListHead->AudioBuffer[PacketSampleIndex]);
+			const float* AudioBuffer = ListHead->AudioBufferMem.GetData();
+			const float* PacketBufferPtr = &(AudioBuffer[PacketSampleIndex]);
 			FMemory::Memcpy(&(DestBuffer[DestBufferSampleIndex]), PacketBufferPtr, NumSamplesLeftInPacketBuffer * sizeof(float));
 			SampleCounter += NumSamplesLeftInPacketBuffer;
 			DestBufferSampleIndex += NumSamplesLeftInPacketBuffer;
@@ -101,7 +102,8 @@ int32 FVoicePacketBuffer::PopAudio(float* DestBuffer, uint32 RequestedSamples)
 		{
 			int32 SamplesLeft = RequestedSamples - DestBufferSampleIndex;
 			const int32 PacketSampleIndex = ListHead->BufferNumSamples - ListHead->SamplesLeft;
-			const float* PacketBufferPtr = &(ListHead->AudioBuffer[PacketSampleIndex]);
+			const float* AudioBuffer = ListHead->AudioBufferMem.GetData();
+			const float* PacketBufferPtr = &(AudioBuffer[PacketSampleIndex]);
 			FMemory::Memcpy(&(DestBuffer[DestBufferSampleIndex]), PacketBufferPtr, SamplesLeft * sizeof(float));
 #if DEBUG_POPPING
 			check(DestBufferSampleIndex + SamplesLeft == RequestedSamples);
diff --git a/Engine/Plugins/Online/OnlineSubsystemUtils/Source/OnlineSubsystemUtils/Public/VoicePacketBuffer.h b/Engine/Plugins/Online/OnlineSubsystemUtils/Source/OnlineSubsystemUtils/Public/VoicePacketBuffer.h
--- a/Engine/Plugins/Online/OnlineSubsystemUtils/Source/OnlineSubsystemUtils/Public/VoicePacketBuffer.h
+++ b/Engine/Plugins/Online/OnlineSubsystemUtils/Source/OnlineSubsystemUtils/Public/VoicePacketBuffer.h
@@ -27,9 +27,6 @@ struct FSortedVoicePacketNode
 	// The actual resulting audio data. allocated ahead of time to minimize memory churn.
 	TArray<float> AudioBufferMem;
 
-	// Cached pointer to AudioBufferMem's internal memory to avoid TArray bounds checks on dereferences.
-	float* AudioBuffer;
-
 	// Number of samples.
 	int32 BufferNumSamples;
 
@@ -56,21 +53,6 @@ struct FSortedVoicePacketNode
 		, IndexInPacketBuffer(INDEX_NONE)
 		, StartSample(0)
 	{
-		AudioBufferMem.Init(0.0f, UVOIPStatics::GetMaxUncompressedVoiceDataSizePerChannel() / sizeof(int16));
-		AudioBuffer = AudioBufferMem.GetData();
-	}
-
-	// Copy constructor.
-	FSortedVoicePacketNode(const FSortedVoicePacketNode& OtherNode)
-		: BufferNumSamples(OtherNode.BufferNumSamples)
-		, SamplesLeft(OtherNode.SamplesLeft)
-		, NextPacket(OtherNode.NextPacket)
-		, IndexInPacketBuffer(OtherNode.IndexInPacketBuffer)
-		, StartSample(OtherNode.StartSample)
-	{
-		AudioBufferMem.Init(0.0f, UVOIPStatics::GetMaxUncompressedVoiceDataSizePerChannel() / sizeof(int16));
-		AudioBuffer = AudioBufferMem.GetData();
-		FMemory::Memcpy(AudioBuffer, OtherNode.AudioBuffer, BufferNumSamples * sizeof(float));
 	}
 
 	// This function is called when a new packet is pushed to avoid reallocating memory.
@@ -81,12 +63,26 @@ struct FSortedVoicePacketNode
 		IndexInPacketBuffer = INDEX_NONE;
 		StartSample = InStartSample;
 
+		if (InFormat == EVoipStreamDataFormat::Int16)
+		{
+			BufferNumSamples = NumBytes / sizeof(int16);
+		}
+		else if (InFormat == EVoipStreamDataFormat::Float)
+		{
+			BufferNumSamples = NumBytes / sizeof(float);
+		}
+		
+
+		AudioBufferMem.Reset();
+		AudioBufferMem.AddUninitialized(BufferNumSamples);
+
+		float* AudioBuffer = AudioBufferMem.GetData();
+
 		// Handle whatever data format the packet is in.
 		switch (InFormat)
 		{
 			case EVoipStreamDataFormat::Float:
 			{
-				BufferNumSamples = NumBytes / sizeof(float);
 				SamplesLeft = BufferNumSamples;
 				FMemory::Memcpy(AudioBuffer, InBuffer, NumBytes);
 				break;
@@ -94,7 +90,6 @@ struct FSortedVoicePacketNode
 
 			case EVoipStreamDataFormat::Int16:
 			{
-				BufferNumSamples = NumBytes / sizeof(int16);
 				SamplesLeft = BufferNumSamples;
 				//Convert to float.
 				int16* BufferPtr = (int16*)InBuffer;
