Fixes for voice engine microphone amplitude.

diff --git a/Engine/Plugins/Online/OnlineSubsystemUtils/Source/OnlineSubsystemUtils/Private/VoiceEngineImpl.cpp b/Engine/Plugins/Online/OnlineSubsystemUtils/Source/OnlineSubsystemUtils/Private/VoiceEngineImpl.cpp
--- a/Engine/Plugins/Online/OnlineSubsystemUtils/Source/OnlineSubsystemUtils/Private/VoiceEngineImpl.cpp
+++ b/Engine/Plugins/Online/OnlineSubsystemUtils/Source/OnlineSubsystemUtils/Private/VoiceEngineImpl.cpp
@@ -430,6 +430,7 @@ uint32 FVoiceEngineImpl::UnregisterRemoteTalker(const FUniqueNetId& UniqueId)
 		// Dump the whole talker
 		RemoteData->Cleanup();
 		RemoteTalkerBuffers.Remove(FUniqueNetIdWrapper(UniqueId.AsShared()));
+		VoiceAmplitudes.Remove(FUniqueNetIdWrapper(UniqueId.AsShared()));
 	}
 
 	return ONLINE_SUCCESS;
@@ -881,11 +882,18 @@ float FVoiceEngineImpl::GetMicrophoneAmplitude(int32 LocalUserNum)
 
 float FVoiceEngineImpl::GetIncomingAudioAmplitude(const FUniqueNetIdWrapper& RemoteUserId)
 {
-	FRemoteTalkerDataImpl* RemoteTalker = RemoteTalkerBuffers.Find(RemoteUserId);
+	FVoiceAmplitudeData* VoiceAmplitude = VoiceAmplitudes.Find(RemoteUserId);
 
-	if (RemoteTalker != nullptr)
+	if (VoiceAmplitude != nullptr)
 	{
-		return RemoteTalker->MicrophoneAmplitude;
+		// Timeout and default to 0 if we haven't received data recently
+		double CurrentTime = FPlatformTime::Seconds();
+		if (CurrentTime - VoiceAmplitude->LastSeen > UVOIPStatics::GetRemoteTalkerTimeoutDuration())
+		{
+			return 0.0f;
+		}
+
+		return VoiceAmplitude->Amplitude;
 	}
 	else
 	{
@@ -895,17 +903,12 @@ float FVoiceEngineImpl::GetIncomingAudioAmplitude(const FUniqueNetIdWrapper& Rem
 
 uint32 FVoiceEngineImpl::SetRemoteVoiceAmplitude(const FUniqueNetIdWrapper& RemoteTalkerId, float InAmplitude)
 {
-	FRemoteTalkerDataImpl* RemoteTalker = RemoteTalkerBuffers.Find(RemoteTalkerId);
+	FVoiceAmplitudeData& VoiceAmplitude = VoiceAmplitudes.FindOrAdd(RemoteTalkerId);
 
-	if (RemoteTalker != nullptr)
-	{
-		RemoteTalker->MicrophoneAmplitude = InAmplitude;
-		return 0;
-	}
-	else
-	{
-		return INDEX_NONE;
-	}
+	VoiceAmplitude.Amplitude = InAmplitude;
+	VoiceAmplitude.LastSeen = FPlatformTime::Seconds();
+
+	return 0;
 }
 
 bool FVoiceEngineImpl::PatchRemoteTalkerOutputToEndpoint(const FString& InDeviceName, bool bMuteInGameOutput /*= true*/)
diff --git a/Engine/Plugins/Online/OnlineSubsystemUtils/Source/OnlineSubsystemUtils/Private/VoicePacketImpl.cpp b/Engine/Plugins/Online/OnlineSubsystemUtils/Source/OnlineSubsystemUtils/Private/VoicePacketImpl.cpp
--- a/Engine/Plugins/Online/OnlineSubsystemUtils/Source/OnlineSubsystemUtils/Private/VoicePacketImpl.cpp
+++ b/Engine/Plugins/Online/OnlineSubsystemUtils/Source/OnlineSubsystemUtils/Private/VoicePacketImpl.cpp
@@ -105,4 +105,5 @@ void FVoicePacketImpl::Serialize(class FArchive& Ar)
 void FVoicePacketImpl::ResetData()
 {
 	Length = 0;
+	MicrophoneAmplitude = 0;
 }
diff --git a/Engine/Plugins/Online/OnlineSubsystemUtils/Source/OnlineSubsystemUtils/Public/VoiceEngineImpl.h b/Engine/Plugins/Online/OnlineSubsystemUtils/Source/OnlineSubsystemUtils/Public/VoiceEngineImpl.h
--- a/Engine/Plugins/Online/OnlineSubsystemUtils/Source/OnlineSubsystemUtils/Public/VoiceEngineImpl.h
+++ b/Engine/Plugins/Online/OnlineSubsystemUtils/Source/OnlineSubsystemUtils/Public/VoiceEngineImpl.h
@@ -43,6 +43,15 @@ struct FLocalVoiceData
 	Audio::FPatchSplitter LocalVoiceOutput;
 };
 
+/**
+ * Container for voice amplitude data
+ */
+struct FVoiceAmplitudeData
+{
+	float Amplitude = 0.0;
+	double LastSeen = 0.0;
+};
+
 /** 
  * Remote voice data playing on a single client
  */
@@ -193,6 +202,9 @@ class ONLINESUBSYSTEMUTILS_API FVoiceEngineImpl : public IVoiceEngine, public FS
 	/** Serialization helper */
 	FVoiceSerializeHelper* SerializeHelper;
 
+	/** Voice Amplitude data to prevent using FRemoteTalkerData if we don't actually require voice*/
+	TMap<FUniqueNetIdWrapper, FVoiceAmplitudeData> VoiceAmplitudes;
+
 	/** Audio taps for the full mixdown of all remote players. */
 	Audio::FPatchMixerSplitter AllRemoteTalkerAudio;
 
