delete packet buffer as soon as the UVoipListenerSynthComponent is marked unreachable or pending kill, so that we don't have to rely on GC throughput.
By instrumenting the UVoipListenerSynthComponent count in a replay, I found that UVoipListenerSynthComponents aren't actually deleted for a very long time after BeginDestroy() is called (though they are eventually deleted en masse). This causes us to accumulate MBs of buffers at a time while we wait for GC to call delete.
Fix is to simply delete the packet buffer as soon as BeginDestroy is called. Added a scope lock also because I am paranoid.

diff --git a/Engine/Plugins/Online/OnlineSubsystemUtils/Source/OnlineSubsystemUtils/Private/VoipListenerSynthComponent.cpp b/Engine/Plugins/Online/OnlineSubsystemUtils/Source/OnlineSubsystemUtils/Private/VoipListenerSynthComponent.cpp
--- a/Engine/Plugins/Online/OnlineSubsystemUtils/Source/OnlineSubsystemUtils/Private/VoipListenerSynthComponent.cpp
+++ b/Engine/Plugins/Online/OnlineSubsystemUtils/Source/OnlineSubsystemUtils/Private/VoipListenerSynthComponent.cpp
@@ -69,6 +69,7 @@ int32 UVoipListenerSynthComponent::OnGenerateAudio(float* OutAudio, int32 NumSam
 	}
 	else if (PacketBuffer.IsValid())
 	{
+		FScopeLock ScopeLock(&PacketBufferCriticalSection);
 		// Handle resync, if neccessary.
 		if (ShouldResyncCVar)
 		{
@@ -177,6 +178,14 @@ uint64 UVoipListenerSynthComponent::GetSampleCounter()
 }
 
 
+void UVoipListenerSynthComponent::BeginDestroy()
+{
+	Super::BeginDestroy();
+	
+	FScopeLock ScopeLock(&PacketBufferCriticalSection);
+	ClosePacketStream();
+}
+
 void UVoipListenerSynthComponent::ForceResync()
 {
 	check(PacketBuffer.IsValid());
diff --git a/Engine/Plugins/Online/OnlineSubsystemUtils/Source/OnlineSubsystemUtils/Public/VoipListenerSynthComponent.h b/Engine/Plugins/Online/OnlineSubsystemUtils/Source/OnlineSubsystemUtils/Public/VoipListenerSynthComponent.h
--- a/Engine/Plugins/Online/OnlineSubsystemUtils/Source/OnlineSubsystemUtils/Public/VoipListenerSynthComponent.h
+++ b/Engine/Plugins/Online/OnlineSubsystemUtils/Source/OnlineSubsystemUtils/Public/VoipListenerSynthComponent.h
@@ -100,6 +100,9 @@ public:
 	*/
 	uint64 GetSampleCounter();
 
+
+	virtual void BeginDestroy() override;
+
 private:
 	/**
 	 * This call will check how much audio we have buffered and drop the oldest audio if necessary.
@@ -112,6 +115,7 @@ private:
 
 	/** This is allocated on OpenPacketStream() */
 	TUniquePtr<FVoicePacketBuffer> PacketBuffer;
+	FCriticalSection PacketBufferCriticalSection;
 
 	float MySampleRate;
 
diff --git a/Engine/Source/Runtime/AudioMixer/Private/Components/SynthComponent.cpp b/Engine/Source/Runtime/AudioMixer/Private/Components/SynthComponent.cpp
--- a/Engine/Source/Runtime/AudioMixer/Private/Components/SynthComponent.cpp
+++ b/Engine/Source/Runtime/AudioMixer/Private/Components/SynthComponent.cpp
@@ -44,8 +44,10 @@ void USynthSound::StartOnAudioDevice(FAudioDevice* InAudioDevice)
 
 void USynthSound::OnBeginGenerate()
 {
-	check(OwningSynthComponent);
-	OwningSynthComponent->OnBeginGenerate();
+	if (ensure(OwningSynthComponent))
+	{
+		OwningSynthComponent->OnBeginGenerate();
+	}
 }
 
 int32 USynthSound::OnGeneratePCMAudio(TArray<uint8>& OutAudio, int32 NumSamples)
@@ -148,6 +150,12 @@ void USynthComponent::OnAudioComponentEnvelopeValue(const UAudioComponent* InAud
 	}
 }
 
+void USynthComponent::BeginDestroy()
+{
+	Super::BeginDestroy();
+	Stop();
+}
+
 void USynthComponent::Activate(bool bReset)
 {
 	if (bReset || ShouldActivate())
diff --git a/Engine/Source/Runtime/AudioMixer/Public/Components/SynthComponent.h b/Engine/Source/Runtime/AudioMixer/Public/Components/SynthComponent.h
--- a/Engine/Source/Runtime/AudioMixer/Public/Components/SynthComponent.h
+++ b/Engine/Source/Runtime/AudioMixer/Public/Components/SynthComponent.h
@@ -216,6 +216,9 @@ public:
 	void AddAudioBufferListener(IAudioBufferListener* InAudioBufferListener);
 	void RemoveAudioBufferListener(IAudioBufferListener* InAudioBufferListener);
 
+
+	virtual void BeginDestroy() override;
+
 protected:
 
 	// Method to execute parameter changes on game thread in audio render thread
