PIE Failures due to invalid credential indexes now give better error logging instead of a generic "Login Failed" which was easily confused with username/password issues.

diff --git a/Engine/Plugins/Online/OnlineSubsystemUtils/Source/OnlineSubsystemUtils/Private/OnlineEngineInterfaceImpl.cpp b/Engine/Plugins/Online/OnlineSubsystemUtils/Source/OnlineSubsystemUtils/Private/OnlineEngineInterfaceImpl.cpp
--- a/Engine/Plugins/Online/OnlineSubsystemUtils/Source/OnlineSubsystemUtils/Private/OnlineEngineInterfaceImpl.cpp
+++ b/Engine/Plugins/Online/OnlineSubsystemUtils/Source/OnlineSubsystemUtils/Private/OnlineEngineInterfaceImpl.cpp
@@ -627,7 +627,7 @@ void UOnlineEngineInterfaceImpl::LoginPIEInstance(FName OnlineIdentifier, int32
 		}
 		else
 		{ 
-			ErrorStr = TEXT("Invalid credentials for PIE login");
+			ErrorStr = FString::Printf(TEXT("Invalid credential index for PIE login. Index: %d NumLogins: %d"), PIELoginNum, PIELogins.Num());
 		}
 	}
 	else
