NetConnection
* Fix infinite hangs when socket errors occur in UIpNetDriver::FReceiveThreadRunnable
* RecvThread will now Yield/Sleep if Wait returns false and error code is SE_NO_ERROR, SE_EWOULDBLOCK or SE_TRY_AGAIN
* RecvThread will now Yield/Sleep if the DispatchQueue is full
* This behavior is disabled if CVarRcvThreadSleepTimeForWaitableErrorsInSeconds is < 0
* Fix warning "NetDriver RecvMulti is not yet supported with the Receive Thread" when only RecvThread is enabled

diff --git a/Engine/Plugins/Online/OnlineSubsystemUtils/Source/OnlineSubsystemUtils/Classes/IpNetDriver.h b/Engine/Plugins/Online/OnlineSubsystemUtils/Source/OnlineSubsystemUtils/Classes/IpNetDriver.h
--- a/Engine/Plugins/Online/OnlineSubsystemUtils/Source/OnlineSubsystemUtils/Classes/IpNetDriver.h
+++ b/Engine/Plugins/Online/OnlineSubsystemUtils/Source/OnlineSubsystemUtils/Classes/IpNetDriver.h
@@ -225,6 +225,9 @@ private:
 		/** Running flag. The Run() function will return shortly after setting this to false. */
 		TAtomic<bool> bIsRunning;
 
+	private:
+		bool DispatchPacket(FReceivedPacket&& IncomingPacket, int32 NbBytesRead);
+
 	private:
 		UIpNetDriver* OwningNetDriver;
 		ISocketSubsystem* SocketSubsystem;
diff --git a/Engine/Plugins/Online/OnlineSubsystemUtils/Source/OnlineSubsystemUtils/Private/IpNetDriver.cpp b/Engine/Plugins/Online/OnlineSubsystemUtils/Source/OnlineSubsystemUtils/Private/IpNetDriver.cpp
--- a/Engine/Plugins/Online/OnlineSubsystemUtils/Source/OnlineSubsystemUtils/Private/IpNetDriver.cpp
+++ b/Engine/Plugins/Online/OnlineSubsystemUtils/Source/OnlineSubsystemUtils/Private/IpNetDriver.cpp
@@ -93,6 +93,11 @@ TAutoConsoleVariable<int32> CVarNetUseRecvTimestamps(
 	TEXT("If true and if net.UseRecvMulti is also true, on a Unix/Linux platform, ")
 		TEXT("the kernel timestamp will be retrieved for each packet received, providing more accurate ping calculations."));
 
+TAutoConsoleVariable<float> CVarRcvThreadSleepTimeForWaitableErrorsInSeconds(
+	TEXT("net.RcvThreadSleepTimeForWaitableErrorsInSeconds"),
+	0.0f, // When > 0 => sleep. When == 0 => yield (if platform supports it). When < 0 => disabled
+	TEXT("Time the receive thread will sleep when a waitable error is returned by a socket operation."));
+
 #if !UE_BUILD_SHIPPING
 TAutoConsoleVariable<int32> CVarNetDebugDualIPs(
 	TEXT("net.DebugDualIPs"),
@@ -105,6 +110,13 @@ TSharedPtr<FInternetAddr> GCurrentDuplicateIP;
 #endif
 
 
+namespace IPNetDriverInternal
+{
+	bool ShouldSleepOnWaitError(ESocketErrors SocketError)
+	{
+		return SocketError == ESocketErrors::SE_NO_ERROR || SocketError == ESocketErrors::SE_EWOULDBLOCK || SocketError == ESocketErrors::SE_TRY_AGAIN;
+	}
+}
 
 /**
  * FPacketItrator
@@ -797,7 +809,7 @@ bool UIpNetDriver::InitBase( bool bInitAsClient, FNetworkNotify* InNotify, const
 			UE_LOG(LogNet, Warning, TEXT("NetDriver could not enable RecvMulti, as current socket subsystem does not support it."));
 		}
 	}
-	else if (bRecvThreadEnabled)
+	else if (bRecvMultiEnabled && bRecvThreadEnabled)
 	{
 		UE_LOG(LogNet, Warning, TEXT("NetDriver RecvMulti is not yet supported with the Receive Thread enabled."));
 	}
@@ -1482,9 +1494,19 @@ UIpNetDriver::FReceiveThreadRunnable::FReceiveThreadRunnable(UIpNetDriver* InOwn
 	SocketSubsystem = OwningNetDriver->GetSocketSubsystem();
 }
 
+bool UIpNetDriver::FReceiveThreadRunnable::DispatchPacket(FReceivedPacket&& IncomingPacket, int32 NbBytesRead)
+{
+	IncomingPacket.PacketBytes.SetNum(FMath::Max(NbBytesRead, 0), false);
+	IncomingPacket.PlatformTimeSeconds = FPlatformTime::Seconds();
+
+	// Add packet to queue. Since ReceiveQueue is a TCircularQueue, if the queue is full, this will simply return false without adding anything.
+	return ReceiveQueue.Enqueue(MoveTemp(IncomingPacket));
+}
+
 uint32 UIpNetDriver::FReceiveThreadRunnable::Run()
 {
 	const FTimespan Timeout = FTimespan::FromMilliseconds(CVarNetIpNetDriverReceiveThreadPollTimeMS.GetValueOnAnyThread());
+	const float SleepTimeForWaitableErrorsInSec = CVarRcvThreadSleepTimeForWaitableErrorsInSeconds.GetValueOnAnyThread();
 
 	UE_LOG(LogNet, Log, TEXT("UIpNetDriver::FReceiveThreadRunnable::Run starting up."));
 
@@ -1492,6 +1514,8 @@ uint32 UIpNetDriver::FReceiveThreadRunnable::Run()
 	{
 		FReceivedPacket IncomingPacket;
 
+		bool bReceiveQueueFull = false;
+
 		if (OwningNetDriver->Socket->Wait(ESocketWaitConditions::WaitForRead, Timeout))
 		{
 			bool bOk = false;
@@ -1508,42 +1532,52 @@ uint32 UIpNetDriver::FReceiveThreadRunnable::Run()
 
 			if (bOk)
 			{
-				if (BytesRead == 0)
+				// Don't even queue empty packets, they can be ignored.
+				if (BytesRead != 0)
 				{
-					// Don't even queue empty packets, they can be ignored.
-					continue;
+					const bool bSuccess = DispatchPacket(MoveTemp(IncomingPacket), BytesRead);
+					bReceiveQueueFull = !bSuccess;
 				}
 			}
 			else
 			{
 				// This relies on the platform's implementation using thread-local storage for the last socket error code.
-				IncomingPacket.Error = SocketSubsystem->GetLastErrorCode();
+				ESocketErrors RecvFromError = SocketSubsystem->GetLastErrorCode();
 
-				// Pass all other errors back to the Game Thread
-				if (IsRecvFailBlocking(IncomingPacket.Error))
+				if (IsRecvFailBlocking(RecvFromError) == false)
 				{
-					continue;
+					// Only non-blocking errors are dispatched to the Game Thread
+					IncomingPacket.Error = RecvFromError;
+					const bool bSuccess = DispatchPacket(MoveTemp(IncomingPacket), BytesRead);
+					bReceiveQueueFull = !bSuccess;
 				}
 			}
-
-
-			IncomingPacket.PacketBytes.SetNum(FMath::Max(BytesRead, 0), false);
-			IncomingPacket.PlatformTimeSeconds = FPlatformTime::Seconds();
-
-			// Add packet to queue. Since ReceiveQueue is a TCircularQueue, if the queue is full, this will simply return false without adding anything.
-			ReceiveQueue.Enqueue(MoveTemp(IncomingPacket));
 		}
 		else
 		{
-			const ESocketErrors WaitError = SocketSubsystem->GetLastErrorCode();
-			if(WaitError != ESocketErrors::SE_NO_ERROR)
+			ESocketErrors WaitError = SocketSubsystem->GetLastErrorCode();
+
+			if (IPNetDriverInternal::ShouldSleepOnWaitError(WaitError))
+			{
+				if (SleepTimeForWaitableErrorsInSec >= 0.0)
+				{
+					FPlatformProcess::SleepNoStats(SleepTimeForWaitableErrorsInSec);
+				}
+			}
+			else if (IsRecvFailBlocking(WaitError) == false)
 			{
+				// Only non-blocking errors are dispatched to the Game Thread
 				IncomingPacket.Error = WaitError;
-				IncomingPacket.PlatformTimeSeconds = FPlatformTime::Seconds();
-
-				UE_LOG(LogNet, Log, TEXT("UIpNetDriver::FReceiveThreadRunnable::Run: Socket->Wait returned error %s (%d)"), SocketSubsystem->GetSocketError(IncomingPacket.Error), static_cast<int>(IncomingPacket.Error));
+				const bool bSuccess = DispatchPacket(MoveTemp(IncomingPacket), 0);
+				bReceiveQueueFull = !bSuccess;
+			}
+		}
 
-				ReceiveQueue.Enqueue(MoveTemp(IncomingPacket));
+		if (bReceiveQueueFull)
+		{
+			if (SleepTimeForWaitableErrorsInSec >= 0.0)
+			{
+				FPlatformProcess::SleepNoStats(SleepTimeForWaitableErrorsInSec);
 			}
 		}
 	}
