Speculative fix for blowing commang queue for SynthComponents.  Added some command culling and setting bAlwaysPlay to true (so command queue gets pumped even if the SynthComponent is silent).

diff --git a/Engine/Plugins/Online/OnlineSubsystemUtils/Source/OnlineSubsystemUtils/Private/OnlineSubsystemUtils.cpp b/Engine/Plugins/Online/OnlineSubsystemUtils/Source/OnlineSubsystemUtils/Private/OnlineSubsystemUtils.cpp
--- a/Engine/Plugins/Online/OnlineSubsystemUtils/Source/OnlineSubsystemUtils/Private/OnlineSubsystemUtils.cpp
+++ b/Engine/Plugins/Online/OnlineSubsystemUtils/Source/OnlineSubsystemUtils/Private/OnlineSubsystemUtils.cpp
@@ -168,11 +168,7 @@ UVoipListenerSynthComponent* CreateVoiceSynthComponent(uint32 SampleRate)
 				SynthComponentPtr->SoundClass = LoadObject<USoundClass>(nullptr, *VoiPSoundClassName.ToString());
 			}
 
-			if (CvarAlwaysPlayVoipComponent)
-			{
-				SynthComponentPtr->bAlwaysPlay = true;
-			}
-
+			SynthComponentPtr->bAlwaysPlay = CvarAlwaysPlayVoipComponent;
 			SynthComponentPtr->Initialize(SampleRate);
 		}
 		else
@@ -202,11 +198,7 @@ UVoipListenerSynthComponent* CreateVoiceSynthComponent(UWorld* World, uint32 Sam
 					SynthComponentPtr->SoundClass = LoadObject<USoundClass>(nullptr, *VoiPSoundClassName.ToString());
 				}
 
-				if (CvarAlwaysPlayVoipComponent)
-				{
-					SynthComponentPtr->bAlwaysPlay = true;
-				}
-
+				SynthComponentPtr->bAlwaysPlay = CvarAlwaysPlayVoipComponent;
 				SynthComponentPtr->RegisterComponentWithWorld(World);
 				SynthComponentPtr->Initialize(SampleRate);
 			}
diff --git a/Engine/Plugins/Runtime/TimeSynth/Source/TimeSynth/Classes/TimeSynthComponent.h b/Engine/Plugins/Runtime/TimeSynth/Source/TimeSynth/Classes/TimeSynthComponent.h
--- a/Engine/Plugins/Runtime/TimeSynth/Source/TimeSynth/Classes/TimeSynthComponent.h
+++ b/Engine/Plugins/Runtime/TimeSynth/Source/TimeSynth/Classes/TimeSynthComponent.h
@@ -697,6 +697,7 @@ private:
 		float TargetVolumeDb;
 		float StartVolumeDb;
 		float CurrentVolumeDb;
+		float LastVolumeDb;
 
 		float CurrentTime;
 		float TargetFadeTime;
@@ -708,6 +709,7 @@ private:
 			: TargetVolumeDb(0.0f)
 			, StartVolumeDb(0.0f)
 			, CurrentVolumeDb(0.0f)
+			, LastVolumeDb(0.0f)
 			, CurrentTime(0.0f)
 			, TargetFadeTime(0.0f)
 		{}
diff --git a/Engine/Plugins/Runtime/TimeSynth/Source/TimeSynth/Private/TimeSynthComponent.cpp b/Engine/Plugins/Runtime/TimeSynth/Source/TimeSynth/Private/TimeSynthComponent.cpp
--- a/Engine/Plugins/Runtime/TimeSynth/Source/TimeSynth/Private/TimeSynthComponent.cpp
+++ b/Engine/Plugins/Runtime/TimeSynth/Source/TimeSynth/Private/TimeSynthComponent.cpp
@@ -233,6 +233,13 @@ void UTimeSynthComponent::TickComponent(float DeltaTime, enum ELevelTick TickTyp
 
 		for (FTimeSynthClipHandle& ClipHandle : VolumeGroup.Clips)
 		{
+			if (FMath::IsNearlyEqual(VolumeGroup.CurrentVolumeDb, VolumeGroup.LastVolumeDb, KINDA_SMALL_NUMBER))
+			{
+				continue;
+			}
+
+			VolumeGroup.LastVolumeDb = VolumeGroup.CurrentVolumeDb;
+
 			float LinearVolume = Audio::ConvertToLinear(VolumeGroup.CurrentVolumeDb);
 
 			SynthCommand([this, ClipHandle, LinearVolume]
diff --git a/Engine/Source/Runtime/AudioMixer/Private/Components/SynthComponent.cpp b/Engine/Source/Runtime/AudioMixer/Private/Components/SynthComponent.cpp
--- a/Engine/Source/Runtime/AudioMixer/Private/Components/SynthComponent.cpp
+++ b/Engine/Source/Runtime/AudioMixer/Private/Components/SynthComponent.cpp
@@ -123,7 +123,7 @@ USynthComponent::USynthComponent(const FObjectInitializer& ObjectInitializer)
 	bIsSynthPlaying = false;
 	bIsInitialized = false;
 	bIsUISound = false;
-	bAlwaysPlay = false;
+	bAlwaysPlay = true;
 	Synth = nullptr;
 
 	// Set the default sound class
