Missing files from Android billing check-in that broke the build.

diff --git a/Engine/Plugins/Online/OnlineSubsystemUtils/Source/OnlineSubsystemUtils/Classes/InAppPurchaseCallbackProxy2.h b/Engine/Plugins/Online/OnlineSubsystemUtils/Source/OnlineSubsystemUtils/Classes/InAppPurchaseCallbackProxy2.h
new file mode 100644
--- /dev/null
+++ b/Engine/Plugins/Online/OnlineSubsystemUtils/Source/OnlineSubsystemUtils/Classes/InAppPurchaseCallbackProxy2.h
@@ -0,0 +1,198 @@
+// Copyright 1998-2019 Epic Games, Inc. All Rights Reserved.
+#pragma once
+
+#include "CoreMinimal.h"
+#include "UObject/ObjectMacros.h"
+#include "UObject/Object.h"
+#include "UObject/ScriptMacros.h"
+#include "Interfaces/OnlinePurchaseInterface.h"
+#include "OnlineSubsystem.h"
+#include "InAppPurchaseCallbackProxy2.generated.h"
+
+/**
+ * Micro-transaction request information
+ */
+USTRUCT(BlueprintType)
+struct FInAppPurchaseProductRequest2
+{
+	GENERATED_USTRUCT_BODY()
+
+		// The unique product identifier that matches the one from your targeted store.
+		UPROPERTY(BlueprintReadWrite, Category = ProductInfo)
+		FString ProductIdentifier;
+
+	// Flag to determine whether this is a consumable purchase, or not.
+	UPROPERTY(BlueprintReadWrite, Category = ProductInfo)
+		bool bIsConsumable;
+};
+
+/**
+ * Micro-transaction purchase information
+ */
+USTRUCT(BlueprintType)
+struct FInAppPurchaseProductInfo2
+{
+	GENERATED_USTRUCT_BODY()
+
+		// The unique product identifier
+		UPROPERTY(BlueprintReadOnly, Category = ProductInfo)
+		FString Identifier;
+
+	// the unique transaction identifier
+	UPROPERTY(BlueprintReadOnly, Category = ProductInfo)
+		FString TransactionIdentifier;
+
+	// The localized display name
+	UPROPERTY(BlueprintReadOnly, Category = ProductInfo)
+		FString DisplayName;
+
+	// The localized display description name
+	UPROPERTY(BlueprintReadOnly, Category = ProductInfo)
+		FString DisplayDescription;
+
+	// The localized display price name
+	UPROPERTY(BlueprintReadOnly, Category = ProductInfo)
+		FString DisplayPrice;
+
+	// Raw price without currency code and symbol
+	UPROPERTY(BlueprintReadOnly, Category = ProductInfo)
+		float RawPrice;
+
+	// The localized currency code of the price
+	UPROPERTY(BlueprintReadOnly, Category = ProductInfo)
+		FString CurrencyCode;
+
+	// The localized currency symbol of the price
+	UPROPERTY(BlueprintReadOnly, Category = ProductInfo)
+		FString CurrencySymbol;
+
+	// The localized decimal separator used in the price
+	UPROPERTY(BlueprintReadOnly, Category = ProductInfo)
+		FString DecimalSeparator;
+
+	// The localized grouping separator of the price
+	UPROPERTY(BlueprintReadOnly, Category = ProductInfo)
+		FString GroupingSeparator;
+
+	// Opaque receipt data for the transaction
+	UPROPERTY(BlueprintReadOnly, Category = ProductInfo)
+		FString ReceiptData;
+
+	// Dynamic fields from raw Json data.
+	UPROPERTY(BlueprintReadOnly, Category = ProductInfo)
+		TMap<FString, FString> DynamicFields;
+};
+
+/**
+ * State of a purchase transaction
+ */
+UENUM(BlueprintType)
+enum class EInAppPurchaseStatus : uint8
+{
+	Invalid = 0 UMETA(DisplayName = "Invalid"),
+	/** purchase completed but failed */
+	Failed UMETA(DisplayName = "Failed"),
+	/** purchase has been deferred (neither failed nor completed) */
+	Deferred UMETA(DisplayName = "Deferred"),
+	/** purchase canceled by user */
+	Canceled UMETA(DisplayName = "Canceled"),
+	/** purchase succeeded */
+	Purchased UMETA(DisplayName = "Purchased"),
+	/** restore succeeded */
+	Restored UMETA(DisplayName = "Restored"),
+};
+
+/**
+ * Micro-transaction purchase information
+ */
+USTRUCT(BlueprintType)
+struct FInAppPurchaseReceiptInfo2
+{
+	GENERATED_USTRUCT_BODY()
+
+	// The item name
+	UPROPERTY(BlueprintReadOnly, Category = ProductInfo)
+		FString ItemName;
+
+	// The unique product identifier
+	UPROPERTY(BlueprintReadOnly, Category = ProductInfo)
+		FString ItemId;
+
+	// the unique transaction identifier
+	UPROPERTY(BlueprintReadOnly, Category = ProductInfo)
+		FString ValidationInfo;
+};
+
+DECLARE_DYNAMIC_MULTICAST_DELEGATE_TwoParams(FInAppPurchaseResult2, EInAppPurchaseStatus, PurchaseStatus, const TArray<FInAppPurchaseReceiptInfo2>&, InAppPurchaseReceipts);
+
+UCLASS(MinimalAPI)
+class UInAppPurchaseCallbackProxy2 : public UObject
+{
+	GENERATED_UCLASS_BODY()
+
+		// Called when there is a successful In-App Purchase transaction
+		UPROPERTY(BlueprintAssignable)
+		FInAppPurchaseResult2 OnSuccess;
+
+		// Called when there is an unsuccessful In-App Purchase transaction
+		UPROPERTY(BlueprintAssignable)
+		FInAppPurchaseResult2 OnFailure;
+
+	// Kicks off a transaction for the provided product identifier
+	UFUNCTION(BlueprintCallable, meta = (DisplayName = "Make an In-App Purchase v2"), Category = "Online|InAppPurchase")
+		static UInAppPurchaseCallbackProxy2* CreateProxyObjectForInAppPurchase(class APlayerController* PlayerController, const FInAppPurchaseProductRequest2& ProductRequest);
+	
+	UFUNCTION(BlueprintCallable, meta = (DisplayName = "Process any New Unprocessed Purchases v2"), Category = "Online|InAppPurchase")
+		static UInAppPurchaseCallbackProxy2* CreateProxyObjectForInAppPurchaseUnprocessedPurchases(class APlayerController* PlayerController);
+	
+	UFUNCTION(BlueprintCallable, meta = (DisplayName = "Query for Owned Purchases"), Category = "Online|InAppPurchase")
+		static UInAppPurchaseCallbackProxy2* CreateProxyObjectForInAppPurchaseQueryOwned(class APlayerController* PlayerController);
+
+public:
+
+	//~ Begin UObject Interface
+	virtual void BeginDestroy() override;
+	//~ End UObject Interface
+
+private:
+
+	/** Called by the InAppPurchase system when the transaction has finished */
+	void OnInAppPurchaseComplete_Delayed();
+	void OnCheckoutComplete(const FOnlineError& Result, const TSharedRef<FPurchaseReceipt>& Receipt);
+	void OnPurchaseComplete();
+	void OnQueryReceiptsComplete(const FOnlineError& Result);
+
+	EInAppPurchaseStatus PurchaseStatusFromOnlineError(const FOnlineError& OnlineError);
+
+	/** Triggers the In-App Purchase Transaction for the specifed user; the Purchase Request object must already be set up */
+	void Trigger(class APlayerController* PlayerController, const FInAppPurchaseProductRequest2& ProductRequest);
+	/** Triggers the In-App Purchase processing for unprocessed receipts */
+	void TriggerGetUnprocessedPurchases(APlayerController* PlayerController);
+	/** Triggers the In-App Purchase processing for unprocessed receipts */
+	void TriggerGetOwnedPurchases(APlayerController* PlayerController);
+
+private:
+
+	/** Delegate called when a InAppPurchase has been successfully read */
+	FOnPurchaseCheckoutComplete InAppPurchaseCompleteDelegate;
+
+	/** Handle to the registered InAppPurchaseCompleteDelegate */
+	FDelegateHandle InAppPurchaseCompleteDelegateHandle;
+
+	/** Did we fail immediately? */
+	bool bFailedToEvenSubmit;
+
+	/** Pointer to the world, needed to delay the results slightly */
+	TWeakObjectPtr<UWorld> WorldPtr;
+
+	/** Did the purchase succeed? */
+	bool bWasSuccessful = false;
+
+	TArray<FInAppPurchaseReceiptInfo2> SavedReceipts;
+
+	TSharedPtr<const FUniqueNetId> PurchasingPlayer;
+
+	IOnlinePurchasePtr PurchaseInterface;
+
+	EInAppPurchaseStatus SavedPurchaseStatus;
+};
\ No newline at end of file
diff --git a/Engine/Plugins/Online/OnlineSubsystemUtils/Source/OnlineSubsystemUtils/Classes/InAppPurchaseQueryCallbackProxy2.h b/Engine/Plugins/Online/OnlineSubsystemUtils/Source/OnlineSubsystemUtils/Classes/InAppPurchaseQueryCallbackProxy2.h
new file mode 100644
--- /dev/null
+++ b/Engine/Plugins/Online/OnlineSubsystemUtils/Source/OnlineSubsystemUtils/Classes/InAppPurchaseQueryCallbackProxy2.h
@@ -0,0 +1,118 @@
+// Copyright 1998-2019 Epic Games, Inc. All Rights Reserved.
+#pragma once
+
+#include "CoreMinimal.h"
+#include "UObject/ObjectMacros.h"
+#include "UObject/Object.h"
+#include "UObject/ScriptMacros.h"
+#include "Interfaces/OnlineStoreInterfaceV2.h"
+#include "InAppPurchaseQueryCallbackProxy2.generated.h"
+
+UENUM(BlueprintType)
+enum class EOnlineProxyStoreOfferDiscountType : uint8
+{
+	/** Offer isn't on sale*/
+	NotOnSale = 0 UMETA(DisplayName = "NotOnSale"),
+	/** Offer price should be displayed as a percentage of regular price */
+	Percentage UMETA(DisplayName = "Percentage"),
+	/** Offer price should be displayed as an amount off regular price */
+	DiscountAmount UMETA(DisplayName = "DiscountAmount"),
+	/** Offer price should be displayed as a new price */
+	PayAmount UMETA(DisplayName = "PayAmount"),
+};
+
+/**
+ * Offer entry for display from online store
+ */
+USTRUCT(BlueprintType)
+struct FOnlineProxyStoreOffer
+{
+	GENERATED_USTRUCT_BODY()
+
+	/** Unique offer identifier */
+	UPROPERTY(BlueprintReadOnly, Category = ProductInfo)
+		FString OfferId;
+
+	/** Title for display */
+	UPROPERTY(BlueprintReadOnly, Category = ProductInfo)
+		FText Title;
+	/** Short description for display */
+	UPROPERTY(BlueprintReadOnly, Category = ProductInfo)
+		FText Description;
+	/** Full description for display */
+	UPROPERTY(BlueprintReadOnly, Category = ProductInfo)
+		FText LongDescription;
+
+	/** Regular non-sale price as text for display */
+	UPROPERTY(BlueprintReadOnly, Category = ProductInfo)
+		FText RegularPriceText;
+	/** Regular non-sale price in numeric form for comparison/sorting */
+	UPROPERTY(BlueprintReadOnly, Category = ProductInfo)
+		int32 RegularPrice;
+
+	/** Final-Pricing (Post-Sales/Discounts) as text for display */
+	UPROPERTY(BlueprintReadOnly, Category = ProductInfo)
+		FText PriceText;
+	/** Final-Price (Post-Sales/Discounts) in numeric form for comparison/sorting */
+	UPROPERTY(BlueprintReadOnly, Category = ProductInfo)
+		int32 NumericPrice;
+
+	/** Price currency code */
+	UPROPERTY(BlueprintReadOnly, Category = ProductInfo)
+		FString CurrencyCode;
+
+	/** Date the offer was released */
+	UPROPERTY(BlueprintReadOnly, Category = ProductInfo)
+		FDateTime ReleaseDate;
+	/** Date this information is no longer valid (maybe due to sale ending, etc) */
+	UPROPERTY(BlueprintReadOnly, Category = ProductInfo)
+		FDateTime ExpirationDate;
+	/** Type of discount currently running on this offer (if any) */
+	UPROPERTY(BlueprintReadOnly, Category = ProductInfo)
+		EOnlineProxyStoreOfferDiscountType DiscountType;
+	UPROPERTY(BlueprintReadOnly, Category = ProductInfo)
+		TMap<FString, FString> DynamicFields;
+};
+
+DECLARE_DYNAMIC_MULTICAST_DELEGATE_OneParam(FInAppPurchaseQuery2Result, const TArray<FOnlineProxyStoreOffer>&, InAppOfferInformation);
+
+UCLASS(MinimalAPI)
+class UInAppPurchaseQueryCallbackProxy2 : public UObject
+{
+	GENERATED_UCLASS_BODY()
+
+	// Called when there is a successful InAppPurchase query
+	UPROPERTY(BlueprintAssignable)
+	FInAppPurchaseQuery2Result OnSuccess;
+
+	// Called when there is an unsuccessful InAppPurchase query
+	UPROPERTY(BlueprintAssignable)
+	FInAppPurchaseQuery2Result OnFailure;
+
+	// Queries a InAppPurchase for an integer value
+	UFUNCTION(BlueprintCallable, meta = (DisplayName="Read In App Purchase Information2"), Category="Online|InAppPurchase")
+	static UInAppPurchaseQueryCallbackProxy2* CreateProxyObjectForInAppPurchaseQuery(class APlayerController* PlayerController, const TArray<FString>& ProductIdentifiers);
+
+public:
+	//~ Begin UObject Interface
+	virtual void BeginDestroy() override;
+	//~ End UObject Interface
+
+private:
+	/** Called by the InAppPurchase system when the read is finished */
+	void OnInAppPurchaseRead_Delayed();
+	void OnInAppPurchaseRead(bool bWasSuccessful, const TArray<FUniqueOfferId>& OfferIds, const FString& Error);
+
+	void CreateProxyProductInformation(TArray<FOnlineStoreOfferRef>& SourceArray, TArray<FOnlineProxyStoreOffer>& TargetArray);
+
+	/** Triggers the query for a specifed user*/
+	void TriggerQuery(class APlayerController* PlayerController, const TArray<FString>& ProductIdentifiers);
+
+private:
+	// Pointer to the world, needed to delay the results slightly
+	TWeakObjectPtr<UWorld> WorldPtr;
+
+	// Did the read succeed?
+	bool bSavedWasSuccessful;
+	TArray<FOnlineProxyStoreOffer> SavedProductInformation;
+};
diff --git a/Engine/Plugins/Online/OnlineSubsystemUtils/Source/OnlineSubsystemUtils/Classes/InAppPurchaseRestoreCallbackProxy2.h b/Engine/Plugins/Online/OnlineSubsystemUtils/Source/OnlineSubsystemUtils/Classes/InAppPurchaseRestoreCallbackProxy2.h
new file mode 100644
--- /dev/null
+++ b/Engine/Plugins/Online/OnlineSubsystemUtils/Source/OnlineSubsystemUtils/Classes/InAppPurchaseRestoreCallbackProxy2.h
@@ -0,0 +1,85 @@
+// Copyright 1998-2019 Epic Games, Inc. All Rights Reserved.
+#pragma once
+
+#include "CoreMinimal.h"
+#include "UObject/ObjectMacros.h"
+#include "UObject/Object.h"
+#include "UObject/ScriptMacros.h"
+#include "Interfaces/OnlinePurchaseInterface.h"
+#include "OnlineSubsystem.h"
+#include "InAppPurchaseCallbackProxy2.h"
+#include "InAppPurchaseRestoreCallbackProxy2.generated.h"
+
+/**
+ * Micro-transaction purchase information
+ */
+USTRUCT(BlueprintType)
+struct FInAppPurchaseRestoreInfo2
+{
+	GENERATED_USTRUCT_BODY()
+
+		// The item name
+		UPROPERTY(BlueprintReadOnly, Category = ProductInfo)
+		FString ItemName;
+
+	// The unique product identifier
+	UPROPERTY(BlueprintReadOnly, Category = ProductInfo)
+		FString ItemId;
+
+	// the unique transaction identifier
+	UPROPERTY(BlueprintReadOnly, Category = ProductInfo)
+		FString ValidationInfo;
+};
+
+DECLARE_DYNAMIC_MULTICAST_DELEGATE_TwoParams(FInAppPurchaseRestoreResult2, EInAppPurchaseStatus, PurchaseStatus, const TArray<FInAppPurchaseRestoreInfo2>&, InAppPurchaseRestoreInfo);
+
+UCLASS(MinimalAPI)
+class UInAppPurchaseRestoreCallbackProxy2 : public UObject
+{
+	GENERATED_UCLASS_BODY()
+
+	// Called when there is a successful In-App Purchase transaction
+	UPROPERTY(BlueprintAssignable)
+	FInAppPurchaseRestoreResult2 OnSuccess;
+
+	// Called when there is an unsuccessful In-App Purchase transaction
+	UPROPERTY(BlueprintAssignable)
+	FInAppPurchaseRestoreResult2 OnFailure;
+
+	// Kicks off a transaction for the provided product identifier
+	UFUNCTION(BlueprintCallable, meta = (DisplayName="Restore In-App Purchases2"), Category="Online|InAppPurchase")
+	static UInAppPurchaseRestoreCallbackProxy2* CreateProxyObjectForInAppPurchaseRestore(const TArray<FInAppPurchaseProductRequest2>& ConsumableProductFlags, class APlayerController* PlayerController);
+
+public:
+
+	//~ Begin UObject Interface
+	virtual void BeginDestroy() override;
+	//~ End UObject Interface
+
+private:
+
+	/** Called by the InAppPurchase system when the transaction has finished */
+	void OnInAppPurchaseRestoreComplete_Delayed();
+	void OnInAppPurchaseRestoreComplete();
+
+	void OnQueryReceiptsComplete(const FOnlineError& Result);
+
+	/** Triggers the In-App Purchase Restore Transaction for the specifed user */
+	void Trigger(const TArray<FInAppPurchaseProductRequest2>& ConsumableProductFlags, class APlayerController* PlayerController);
+
+	EInAppPurchaseStatus PurchaseStatusFromOnlineError(const FOnlineError& OnlineError);
+
+private:
+
+	bool bWasSuccessful;
+	/** Pointer to the world, needed to delay the results slightly */
+	TWeakObjectPtr<UWorld> WorldPtr;
+
+	TArray<FInAppPurchaseRestoreInfo2> SavedReceipts;
+
+	IOnlinePurchasePtr PurchaseInterface;
+
+	FUniqueNetIdRepl PurchasingPlayer;
+	
+	EInAppPurchaseStatus SavedPurchaseStatus;
+};
diff --git a/Engine/Plugins/Online/OnlineSubsystemUtils/Source/OnlineSubsystemUtils/Private/InAppPurchaseCallbackProxy2.cpp b/Engine/Plugins/Online/OnlineSubsystemUtils/Source/OnlineSubsystemUtils/Private/InAppPurchaseCallbackProxy2.cpp
new file mode 100644
--- /dev/null
+++ b/Engine/Plugins/Online/OnlineSubsystemUtils/Source/OnlineSubsystemUtils/Private/InAppPurchaseCallbackProxy2.cpp
@@ -0,0 +1,276 @@
+// Copyright 1998-2019 Epic Games, Inc. All Rights Reserved.
+
+#include "InAppPurchaseCallbackProxy2.h"
+#include "Async/TaskGraphInterfaces.h"
+#include "GameFramework/PlayerController.h"
+#include "OnlineSubsystem.h"
+#include "Engine/World.h"
+
+//////////////////////////////////////////////////////////////////////////
+// UInAppPurchaseCallbackProxy
+
+UInAppPurchaseCallbackProxy2::UInAppPurchaseCallbackProxy2(const FObjectInitializer& ObjectInitializer)
+	: Super(ObjectInitializer)
+{
+	WorldPtr = nullptr;
+	bWasSuccessful = false;
+	SavedPurchaseStatus = EInAppPurchaseStatus::Failed;
+}
+
+
+void UInAppPurchaseCallbackProxy2::Trigger(APlayerController* PlayerController, const FInAppPurchaseProductRequest2& ProductRequest)
+{
+	bFailedToEvenSubmit = true;
+	WorldPtr = (PlayerController != nullptr) ? PlayerController->GetWorld() : nullptr;
+	if (APlayerState* PlayerState = (PlayerController != nullptr) ? PlayerController->PlayerState : nullptr)
+	{
+		if (IOnlineSubsystem* const OnlineSub = IOnlineSubsystem::IsLoaded() ? IOnlineSubsystem::Get() : nullptr)
+		{
+			PurchaseInterface = OnlineSub->GetPurchaseInterface();
+			if (PurchaseInterface.IsValid())
+			{
+				bFailedToEvenSubmit = false;
+
+				// Register the completion callback
+				InAppPurchaseCompleteDelegate = FOnPurchaseCheckoutComplete::CreateUObject(this, &UInAppPurchaseCallbackProxy2::OnCheckoutComplete);
+
+				// Set-up, and trigger the transaction through the store interface
+				FPurchaseCheckoutRequest CheckoutRequest = FPurchaseCheckoutRequest();
+				CheckoutRequest.AddPurchaseOffer("", ProductRequest.ProductIdentifier, 1, ProductRequest.bIsConsumable);
+				PurchasingPlayer = (*PlayerController->GetLocalPlayer()->GetUniqueNetIdFromCachedControllerId()).AsShared();
+				PurchaseInterface->Checkout(*PurchasingPlayer, CheckoutRequest, InAppPurchaseCompleteDelegate);
+			}
+			else
+			{
+				FFrame::KismetExecutionMessage(TEXT("UInAppPurchaseCallbackProxy::Trigger - In-App Purchases are not supported by Online Subsystem"), ELogVerbosity::Warning);
+			}
+		}
+		else
+		{
+			FFrame::KismetExecutionMessage(TEXT("UInAppPurchaseCallbackProxy::Trigger - Invalid or uninitialized OnlineSubsystem"), ELogVerbosity::Warning);
+		}
+	}
+	else
+	{
+		FFrame::KismetExecutionMessage(TEXT("UInAppPurchaseCallbackProxy::Trigger - Invalid player state"), ELogVerbosity::Warning);
+	}
+
+	if (bFailedToEvenSubmit && (PlayerController != NULL))
+	{
+		OnPurchaseComplete();
+	}
+}
+
+void UInAppPurchaseCallbackProxy2::TriggerGetUnprocessedPurchases(APlayerController* PlayerController)
+{
+	WorldPtr = (PlayerController != nullptr) ? PlayerController->GetWorld() : nullptr;
+	if (APlayerState* PlayerState = (PlayerController != nullptr) ? PlayerController->PlayerState : nullptr)
+	{
+		if (IOnlineSubsystem* const OnlineSub = IOnlineSubsystem::IsLoaded() ? IOnlineSubsystem::Get() : nullptr)
+		{
+			PurchaseInterface = OnlineSub->GetPurchaseInterface();
+			if (PurchaseInterface.IsValid())
+			{
+				PurchasingPlayer = (*PlayerController->GetLocalPlayer()->GetUniqueNetIdFromCachedControllerId()).AsShared();
+				OnCheckoutComplete(FOnlineError(true), MakeShared<FPurchaseReceipt>());
+			}
+			else
+			{
+				FFrame::KismetExecutionMessage(TEXT("UInAppPurchaseCallbackProxy::Trigger - In-App Purchases are not supported by Online Subsystem"), ELogVerbosity::Warning);
+			}
+		}
+		else
+		{
+			FFrame::KismetExecutionMessage(TEXT("UInAppPurchaseCallbackProxy::Trigger - Invalid or uninitialized OnlineSubsystem"), ELogVerbosity::Warning);
+		}
+	}
+	else
+	{
+		FFrame::KismetExecutionMessage(TEXT("UInAppPurchaseCallbackProxy::Trigger - Invalid player state"), ELogVerbosity::Warning);
+	}
+}
+
+void UInAppPurchaseCallbackProxy2::TriggerGetOwnedPurchases(APlayerController* PlayerController)
+{
+	bFailedToEvenSubmit = true;
+	WorldPtr = (PlayerController != nullptr) ? PlayerController->GetWorld() : nullptr;
+	if (APlayerState* PlayerState = (PlayerController != nullptr) ? PlayerController->PlayerState : nullptr)
+	{
+		if (IOnlineSubsystem* const OnlineSub = IOnlineSubsystem::IsLoaded() ? IOnlineSubsystem::Get() : nullptr)
+		{
+			PurchaseInterface = OnlineSub->GetPurchaseInterface();
+			if (PurchaseInterface.IsValid())
+			{
+				PurchasingPlayer = (*PlayerController->GetLocalPlayer()->GetUniqueNetIdFromCachedControllerId()).AsShared();
+				PurchaseInterface->QueryReceipts(*PurchasingPlayer, false, FOnQueryReceiptsComplete::CreateUObject(this, &UInAppPurchaseCallbackProxy2::OnQueryReceiptsComplete));
+			}
+			else
+			{
+				FFrame::KismetExecutionMessage(TEXT("UInAppPurchaseCallbackProxy::Trigger - In-App Purchases are not supported by Online Subsystem"), ELogVerbosity::Warning);
+			}
+		}
+		else
+		{
+			FFrame::KismetExecutionMessage(TEXT("UInAppPurchaseCallbackProxy::Trigger - Invalid or uninitialized OnlineSubsystem"), ELogVerbosity::Warning);
+		}
+	}
+	else
+	{
+		FFrame::KismetExecutionMessage(TEXT("UInAppPurchaseCallbackProxy::Trigger - Invalid player state"), ELogVerbosity::Warning);
+	}
+
+	if (bFailedToEvenSubmit && (PlayerController != NULL))
+	{
+		OnPurchaseComplete();
+	}
+}
+
+void UInAppPurchaseCallbackProxy2::OnCheckoutComplete(const FOnlineError& Result, const TSharedRef<FPurchaseReceipt>& Receipt)
+{
+	
+	TArray<FPurchaseReceipt> OutReceipts;
+	bWasSuccessful = false;
+	bool bWasSuccessfulLocal = Result.WasSuccessful();
+	SavedPurchaseStatus = PurchaseStatusFromOnlineError(Result);
+	if (bWasSuccessfulLocal)
+	{
+		if (PurchaseInterface.IsValid())
+		{
+			if (PurchasingPlayer.IsValid())
+			{
+				// Get Receipts
+				PurchaseInterface->GetReceipts(*PurchasingPlayer, OutReceipts);
+				for (FPurchaseReceipt& CurrentReceipt : OutReceipts)
+				{
+					if (CurrentReceipt.TransactionState == EPurchaseTransactionState::Purchased ||
+						CurrentReceipt.TransactionState == EPurchaseTransactionState::Restored)
+					{
+						FFrame::KismetExecutionMessage(TEXT("UInAppPurchaseCallbackProxy::OnCheckoutComplete - FinalizingPurchase"), ELogVerbosity::Warning);
+						PurchaseInterface->FinalizePurchase(*PurchasingPlayer, CurrentReceipt.TransactionId);
+						bWasSuccessful = bWasSuccessfulLocal;
+					}
+				}
+			}
+			else
+			{
+				FFrame::KismetExecutionMessage(TEXT("UInAppPurchaseCallbackProxy::OnCheckoutComplete - Invalid or uninitialized PurchasingPlayer"), ELogVerbosity::Warning);
+			}
+		}
+	}
+
+	FFrame::KismetExecutionMessage(TEXT("UInAppPurchaseCallbackProxy::OnCheckoutComplete - Emptying savedreceipts"), ELogVerbosity::Warning);
+	SavedReceipts.Empty();
+
+	FFrame::KismetExecutionMessage(TEXT("UInAppPurchaseCallbackProxy::OnCheckoutComplete - Iterating OutReceipts"), ELogVerbosity::Warning);
+	for(FPurchaseReceipt& CurrentReceipt : OutReceipts)
+	{
+		if (CurrentReceipt.TransactionState == EPurchaseTransactionState::Purchased ||
+			CurrentReceipt.TransactionState == EPurchaseTransactionState::Restored)
+		{
+			FInAppPurchaseReceiptInfo2& ReceiptInfo = SavedReceipts.AddDefaulted_GetRef();
+
+			FFrame::KismetExecutionMessage(TEXT("UInAppPurchaseCallbackProxy::OnCheckoutComplete - Writing Receipt to savedreceipts"), ELogVerbosity::Warning);
+			ReceiptInfo.ItemId = CurrentReceipt.ReceiptOffers[0].LineItems[0].UniqueId;
+			ReceiptInfo.ItemName = CurrentReceipt.ReceiptOffers[0].LineItems[0].ItemName;
+			ReceiptInfo.ValidationInfo = CurrentReceipt.ReceiptOffers[0].LineItems[0].ValidationInfo;
+		}
+	}
+	FFrame::KismetExecutionMessage(TEXT("UInAppPurchaseCallbackProxy::OnCheckoutComplete - Calling OnPurchaseCOmplete"), ELogVerbosity::Warning);
+	OnPurchaseComplete();
+}
+
+void UInAppPurchaseCallbackProxy2::OnPurchaseComplete()
+{
+	if (UWorld* World = WorldPtr.Get())
+	{
+		DECLARE_CYCLE_STAT(TEXT("FSimpleDelegateGraphTask.DelayInAppPurchaseComplete"), STAT_FSimpleDelegateGraphTask_DelayInAppPurchaseComplete, STATGROUP_TaskGraphTasks);
+
+		FSimpleDelegateGraphTask::CreateAndDispatchWhenReady(
+			FSimpleDelegateGraphTask::FDelegate::CreateLambda([=]() {
+
+			OnInAppPurchaseComplete_Delayed();
+
+		}),
+			GET_STATID(STAT_FSimpleDelegateGraphTask_DelayInAppPurchaseComplete),
+			nullptr,
+			ENamedThreads::GameThread
+			);
+	}
+}
+
+EInAppPurchaseStatus UInAppPurchaseCallbackProxy2::PurchaseStatusFromOnlineError(const FOnlineError& OnlineError)
+{
+	if (OnlineError.bSucceeded)
+	{
+		return EInAppPurchaseStatus::Purchased;
+	}
+	else if(OnlineError.ErrorCode.Equals(TEXT("com.epicgames.purchase.failure")))
+	{
+		return EInAppPurchaseStatus::Failed;
+	}
+	else if (OnlineError.ErrorCode.Equals(TEXT("com.epicgames.catalog_helper.user_cancelled")))
+	{
+		return EInAppPurchaseStatus::Canceled;
+	}
+	else if (OnlineError.ErrorCode.Equals(TEXT("com.epicgames.purchase.deferred")))
+	{
+		return EInAppPurchaseStatus::Deferred;
+	}
+	else if (OnlineError.ErrorCode.Equals(TEXT("com.epicgames.purchase.invalid")))
+	{
+		return EInAppPurchaseStatus::Invalid;
+	}
+	
+	return EInAppPurchaseStatus::Failed;			
+}
+
+void UInAppPurchaseCallbackProxy2::OnQueryReceiptsComplete(const FOnlineError& Result)
+{
+		SavedReceipts.Empty();
+		bWasSuccessful = Result.WasSuccessful();
+		SavedPurchaseStatus = PurchaseStatusFromOnlineError(Result);
+		
+		OnPurchaseComplete();
+}
+
+void UInAppPurchaseCallbackProxy2::OnInAppPurchaseComplete_Delayed()
+{
+	/** Cached product details of the purchased product */
+	if (bWasSuccessful)
+	{
+		OnSuccess.Broadcast(SavedPurchaseStatus, SavedReceipts);
+	}
+	else
+	{
+		OnFailure.Broadcast(SavedPurchaseStatus, SavedReceipts);
+	}
+}
+
+void UInAppPurchaseCallbackProxy2::BeginDestroy()
+{
+	Super::BeginDestroy();
+}
+
+UInAppPurchaseCallbackProxy2* UInAppPurchaseCallbackProxy2::CreateProxyObjectForInAppPurchase(class APlayerController* PlayerController, const FInAppPurchaseProductRequest2& ProductRequest)
+{
+	UInAppPurchaseCallbackProxy2* Proxy = NewObject<UInAppPurchaseCallbackProxy2>();
+	Proxy->SetFlags(RF_StrongRefOnFrame);
+	Proxy->Trigger(PlayerController, ProductRequest);
+	return Proxy;
+}
+
+UInAppPurchaseCallbackProxy2* UInAppPurchaseCallbackProxy2::CreateProxyObjectForInAppPurchaseUnprocessedPurchases(class APlayerController* PlayerController)
+{
+	UInAppPurchaseCallbackProxy2* Proxy = NewObject<UInAppPurchaseCallbackProxy2>();
+	Proxy->SetFlags(RF_StrongRefOnFrame);
+	Proxy->TriggerGetUnprocessedPurchases(PlayerController);
+	return Proxy;
+}
+
+UInAppPurchaseCallbackProxy2* UInAppPurchaseCallbackProxy2::CreateProxyObjectForInAppPurchaseQueryOwned(class APlayerController* PlayerController)
+{
+	UInAppPurchaseCallbackProxy2* Proxy = NewObject<UInAppPurchaseCallbackProxy2>();
+	Proxy->SetFlags(RF_StrongRefOnFrame);
+	Proxy->TriggerGetOwnedPurchases(PlayerController);
+	return Proxy;
+}
\ No newline at end of file
diff --git a/Engine/Plugins/Online/OnlineSubsystemUtils/Source/OnlineSubsystemUtils/Private/InAppPurchaseQueryCallbackProxy2.cpp b/Engine/Plugins/Online/OnlineSubsystemUtils/Source/OnlineSubsystemUtils/Private/InAppPurchaseQueryCallbackProxy2.cpp
new file mode 100644
--- /dev/null
+++ b/Engine/Plugins/Online/OnlineSubsystemUtils/Source/OnlineSubsystemUtils/Private/InAppPurchaseQueryCallbackProxy2.cpp
@@ -0,0 +1,158 @@
+// Copyright 1998-2019 Epic Games, Inc. All Rights Reserved.
+
+#include "InAppPurchaseQueryCallbackProxy2.h"
+#include "Async/TaskGraphInterfaces.h"
+#include "GameFramework/PlayerController.h"
+#include "OnlineSubsystem.h"
+#include "Engine/World.h"
+
+//////////////////////////////////////////////////////////////////////////
+// UInAppPurchaseQueryCallbackProxy
+
+UInAppPurchaseQueryCallbackProxy2::UInAppPurchaseQueryCallbackProxy2(const FObjectInitializer& ObjectInitializer)
+	: Super(ObjectInitializer)
+{
+}
+
+void UInAppPurchaseQueryCallbackProxy2::TriggerQuery(APlayerController* PlayerController, const TArray<FString>& ProductIdentifiers)
+{
+	bool bFailedToEvenSubmit = true;
+
+	WorldPtr = (PlayerController != nullptr) ? PlayerController->GetWorld() : nullptr;
+	if (APlayerState* PlayerState = (PlayerController != NULL) ? PlayerController->PlayerState : nullptr)
+	{
+		if (IOnlineSubsystem* const OnlineSub = IOnlineSubsystem::IsLoaded() ? IOnlineSubsystem::Get() : nullptr)
+		{
+			IOnlineStoreV2Ptr StoreInterface = OnlineSub->GetStoreV2Interface();
+			if (StoreInterface.IsValid())
+			{
+				FFrame::KismetExecutionMessage(TEXT("UInAppPurchaseQueryCallbackProxy2::TriggerQuery - GetUniqueNetIdFromCachedControllerId"), ELogVerbosity::Warning);
+				FUniqueNetIdRepl QueryingPlayer = PlayerController->GetLocalPlayer()->GetUniqueNetIdFromCachedControllerId();
+
+				//if (QueryingPlayer.IsValid())
+				//{
+				FFrame::KismetExecutionMessage(TEXT("UInAppPurchaseQueryCallbackProxy2::TriggerQuery - Querying Store Interface"), ELogVerbosity::Warning);
+					bFailedToEvenSubmit = false;
+					StoreInterface->QueryOffersById(*QueryingPlayer, ProductIdentifiers, FOnQueryOnlineStoreOffersComplete::CreateUObject(this, &UInAppPurchaseQueryCallbackProxy2::OnInAppPurchaseRead));
+				//}
+				//else
+				//{
+				//	FFrame::KismetExecutionMessage(TEXT("UInAppPurchaseQueryCallbackProxy2::TriggerQuery - Invalid UniqueNetId"), ELogVerbosity::Warning);
+				//}
+			}
+			else
+			{
+				FFrame::KismetExecutionMessage(TEXT("UInAppPurchaseQueryCallbackProxy2::TriggerQuery - In App Purchases are not supported by Online Subsystem"), ELogVerbosity::Warning);
+			}
+		}
+		else
+		{
+			FFrame::KismetExecutionMessage(TEXT("UInAppPurchaseQueryCallbackProxy2::TriggerQuery - Invalid or uninitialized OnlineSubsystem"), ELogVerbosity::Warning);
+		}
+	}
+	else
+	{
+		FFrame::KismetExecutionMessage(TEXT("UInAppPurchaseQueryCallbackProxy2::TriggerQuery - Invalid player state"), ELogVerbosity::Warning);
+	}
+
+	if (bFailedToEvenSubmit && (PlayerController != NULL))
+	{
+		FFrame::KismetExecutionMessage(TEXT("UInAppPurchaseQueryCallbackProxy2::TriggerQuery - Failed to even submit"), ELogVerbosity::Warning);
+		OnInAppPurchaseRead(false, TArray<FUniqueOfferId>(), FString());
+	}
+}
+
+void UInAppPurchaseQueryCallbackProxy2::CreateProxyProductInformation(TArray<FOnlineStoreOfferRef>& SourceArray, TArray<FOnlineProxyStoreOffer>& TargetArray)
+{
+	FFrame::KismetExecutionMessage(TEXT("UInAppPurchaseQueryCallbackProxy2::TriggerQuery - CreateProxyInformation"), ELogVerbosity::Warning);
+	TargetArray.Empty();
+	for (FOnlineStoreOfferRef SourceProduct : SourceArray)
+	{
+		FFrame::KismetExecutionMessage(TEXT("UInAppPurchaseQueryCallbackProxy2::TriggerQuery - CreateProxyInformation In Loop"), ELogVerbosity::Warning);
+		FOnlineProxyStoreOffer CurrentProduct;
+
+		CurrentProduct.OfferId = SourceProduct->OfferId;
+		CurrentProduct.Title = SourceProduct->Title;
+		CurrentProduct.Description = SourceProduct->Description;
+		CurrentProduct.LongDescription = SourceProduct->LongDescription;
+		CurrentProduct.RegularPriceText = SourceProduct->RegularPriceText;
+		CurrentProduct.RegularPrice = SourceProduct->RegularPrice;
+		CurrentProduct.PriceText = SourceProduct->PriceText;
+		CurrentProduct.NumericPrice = SourceProduct->NumericPrice;
+		CurrentProduct.CurrencyCode = SourceProduct->CurrencyCode;
+		CurrentProduct.ReleaseDate = SourceProduct->ReleaseDate;
+		CurrentProduct.ExpirationDate = SourceProduct->ExpirationDate;
+		CurrentProduct.DiscountType = (EOnlineProxyStoreOfferDiscountType)SourceProduct->DiscountType;
+		CurrentProduct.DynamicFields = SourceProduct->DynamicFields;
+
+		FFrame::KismetExecutionMessage(TEXT("UInAppPurchaseQueryCallbackProxy2::TriggerQuery - Adding Product"), ELogVerbosity::Warning);
+		TargetArray.Add(CurrentProduct);
+		FFrame::KismetExecutionMessage(TEXT("UInAppPurchaseQueryCallbackProxy2::TriggerQuery - Done Adding Product"), ELogVerbosity::Warning);
+	}
+}
+
+void UInAppPurchaseQueryCallbackProxy2::OnInAppPurchaseRead(bool bWasSuccessful, const TArray<FUniqueOfferId>& OfferIds, const FString& Error)
+{
+	FFrame::KismetExecutionMessage(TEXT("UInAppPurchaseQueryCallbackProxy2::OnInAppPurchaseRead - Start"), ELogVerbosity::Warning);
+	bSavedWasSuccessful = false;
+	if (bWasSuccessful)
+	{
+		FFrame::KismetExecutionMessage(TEXT("UInAppPurchaseQueryCallbackProxy2::OnInAppPurchaseRead - WasSuccessful"), ELogVerbosity::Warning);
+		if (IOnlineSubsystem* const OnlineSub = IOnlineSubsystem::IsLoaded() ? IOnlineSubsystem::Get() : nullptr)
+		{
+			IOnlineStoreV2Ptr StoreInterface = OnlineSub->GetStoreV2Interface();
+			if (StoreInterface.IsValid())
+			{
+				FFrame::KismetExecutionMessage(TEXT("UInAppPurchaseQueryCallbackProxy2::OnInAppPurchaseRead - Interface is good"), ELogVerbosity::Warning);
+				TArray<FOnlineStoreOfferRef> LocalProductInformation;
+				FFrame::KismetExecutionMessage(TEXT("UInAppPurchaseQueryCallbackProxy2::OnInAppPurchaseRead - Getting Cached Offers"), ELogVerbosity::Warning);
+				StoreInterface->GetOffers(LocalProductInformation);
+				FFrame::KismetExecutionMessage(TEXT("UInAppPurchaseQueryCallbackProxy2::OnInAppPurchaseRead - Converting to proxy store offer"), ELogVerbosity::Warning);
+				CreateProxyProductInformation(LocalProductInformation, SavedProductInformation);
+				FFrame::KismetExecutionMessage(TEXT("UInAppPurchaseQueryCallbackProxy2::OnInAppPurchaseRead - Done"), ELogVerbosity::Warning);
+				bSavedWasSuccessful = true;
+			}
+		}
+	}
+
+	if (UWorld* World = WorldPtr.Get())
+	{
+		DECLARE_CYCLE_STAT(TEXT("FSimpleDelegateGraphTask.DelayInAppPurchaseRead"), STAT_FSimpleDelegateGraphTask_DelayInAppPurchaseRead, STATGROUP_TaskGraphTasks);
+
+		FSimpleDelegateGraphTask::CreateAndDispatchWhenReady(
+			FSimpleDelegateGraphTask::FDelegate::CreateLambda([=](){
+
+				OnInAppPurchaseRead_Delayed();
+
+			}),
+			GET_STATID(STAT_FSimpleDelegateGraphTask_DelayInAppPurchaseRead), 
+			nullptr, 
+			ENamedThreads::GameThread
+		);
+	}
+}
+
+void UInAppPurchaseQueryCallbackProxy2::OnInAppPurchaseRead_Delayed()
+{
+	if (bSavedWasSuccessful)
+	{
+		OnSuccess.Broadcast(SavedProductInformation);
+	}
+	else
+	{
+		OnFailure.Broadcast(SavedProductInformation);
+	}
+}
+
+void UInAppPurchaseQueryCallbackProxy2::BeginDestroy()
+{
+	Super::BeginDestroy();
+}
+
+UInAppPurchaseQueryCallbackProxy2* UInAppPurchaseQueryCallbackProxy2::CreateProxyObjectForInAppPurchaseQuery(class APlayerController* PlayerController, const TArray<FString>& ProductIdentifiers)
+{
+	UInAppPurchaseQueryCallbackProxy2* Proxy = NewObject<UInAppPurchaseQueryCallbackProxy2>();
+	Proxy->SetFlags(RF_StrongRefOnFrame);
+	Proxy->TriggerQuery(PlayerController, ProductIdentifiers);
+	return Proxy;
+}
diff --git a/Engine/Plugins/Online/OnlineSubsystemUtils/Source/OnlineSubsystemUtils/Private/InAppPurchaseRestoreCallbackProxy2.cpp b/Engine/Plugins/Online/OnlineSubsystemUtils/Source/OnlineSubsystemUtils/Private/InAppPurchaseRestoreCallbackProxy2.cpp
new file mode 100644
--- /dev/null
+++ b/Engine/Plugins/Online/OnlineSubsystemUtils/Source/OnlineSubsystemUtils/Private/InAppPurchaseRestoreCallbackProxy2.cpp
@@ -0,0 +1,166 @@
+// Copyright 1998-2019 Epic Games, Inc. All Rights Reserved.
+
+#include "InAppPurchaseRestoreCallbackProxy2.h"
+#include "Async/TaskGraphInterfaces.h"
+#include "GameFramework/PlayerController.h"
+#include "OnlineSubsystem.h"
+#include "Engine/World.h"
+
+//////////////////////////////////////////////////////////////////////////
+// UInAppPurchaseRestoreCallbackProxy
+
+UInAppPurchaseRestoreCallbackProxy2::UInAppPurchaseRestoreCallbackProxy2(const FObjectInitializer& ObjectInitializer)
+	: Super(ObjectInitializer)
+{
+	WorldPtr = nullptr;
+}
+
+
+void UInAppPurchaseRestoreCallbackProxy2::Trigger(const TArray<FInAppPurchaseProductRequest2>& ConsumableProductFlags, APlayerController* PlayerController)
+{
+	bWasSuccessful = false;
+	bool bFailedToEvenSubmit = true;
+
+	WorldPtr = (PlayerController != nullptr) ? PlayerController->GetWorld() : nullptr;
+	if (APlayerState* PlayerState = (PlayerController != nullptr) ? PlayerController->PlayerState : nullptr)
+	{
+		if (IOnlineSubsystem* const OnlineSub = IOnlineSubsystem::IsLoaded() ? IOnlineSubsystem::Get() : nullptr)
+		{
+			PurchaseInterface = OnlineSub->GetPurchaseInterface();
+			if (PurchaseInterface.IsValid())
+			{
+				bFailedToEvenSubmit = false;
+
+				PurchasingPlayer = (*PlayerController->GetLocalPlayer()->GetUniqueNetIdFromCachedControllerId()).AsShared();
+				PurchaseInterface->QueryReceipts(*PurchasingPlayer, true, FOnQueryReceiptsComplete::CreateUObject(this, &UInAppPurchaseRestoreCallbackProxy2::OnQueryReceiptsComplete));
+			}
+			else
+			{
+				FFrame::KismetExecutionMessage(TEXT("UInAppPurchaseRestoreCallbackProxy2::Trigger - In-App Purchases are not supported by Online Subsystem"), ELogVerbosity::Warning);
+			}
+		}
+		else
+		{
+			FFrame::KismetExecutionMessage(TEXT("UInAppPurchaseRestoreCallbackProxy2::Trigger - Invalid or uninitialized OnlineSubsystem"), ELogVerbosity::Warning);
+		}
+	}
+	else
+	{
+		FFrame::KismetExecutionMessage(TEXT("UInAppPurchaseRestoreCallbackProxy2::Trigger - Invalid player state"), ELogVerbosity::Warning);
+	}
+
+	if (bFailedToEvenSubmit && (PlayerController != NULL))
+	{
+		bWasSuccessful = false;
+		OnInAppPurchaseRestoreComplete();
+	}
+}
+
+void UInAppPurchaseRestoreCallbackProxy2::OnQueryReceiptsComplete(const FOnlineError& Result)
+{
+	
+	bWasSuccessful = Result.WasSuccessful();
+	SavedPurchaseStatus = PurchaseStatusFromOnlineError(Result);
+	check(PurchaseInterface.IsValid());
+
+	// Get Receipts
+	TArray<FPurchaseReceipt> OutReceipts;
+	PurchaseInterface->GetReceipts(*PurchasingPlayer, OutReceipts);
+	for (FPurchaseReceipt& CurrentReceipt : OutReceipts)
+	{
+		if (CurrentReceipt.TransactionState == EPurchaseTransactionState::Restored)
+		{
+			SavedPurchaseStatus = EInAppPurchaseStatus::Restored;
+		}
+	}
+
+	SavedReceipts.Empty();
+
+	FFrame::KismetExecutionMessage(TEXT("UInAppPurchaseRestoreCallbackProxy2::OnQueryReceiptsComplete - Iterating OutReceipts"), ELogVerbosity::Warning);
+	for (FPurchaseReceipt& CurrentReceipt : OutReceipts)
+	{
+		if (CurrentReceipt.TransactionState == EPurchaseTransactionState::Restored)
+		{
+			FInAppPurchaseRestoreInfo2& ReceiptInfo = SavedReceipts.AddDefaulted_GetRef();
+
+			FFrame::KismetExecutionMessage(TEXT("UInAppPurchaseRestoreCallbackProxy2::OnQueryReceiptsComplete - Writing Receipt to savedreceipts"), ELogVerbosity::Warning);
+			ReceiptInfo.ItemId = CurrentReceipt.ReceiptOffers[0].LineItems[0].UniqueId;
+			ReceiptInfo.ItemName = CurrentReceipt.ReceiptOffers[0].LineItems[0].ItemName;
+			ReceiptInfo.ValidationInfo = CurrentReceipt.ReceiptOffers[0].LineItems[0].ValidationInfo;
+		}
+	}
+
+	OnInAppPurchaseRestoreComplete();
+}
+
+void UInAppPurchaseRestoreCallbackProxy2::OnInAppPurchaseRestoreComplete()
+{
+	if (UWorld* World = WorldPtr.Get())
+	{
+		DECLARE_CYCLE_STAT(TEXT("FSimpleDelegateGraphTask.DelayInAppPurchaseRestoreComplete"), STAT_FSimpleDelegateGraphTask_DelayInAppPurchaseRestoreComplete, STATGROUP_TaskGraphTasks);
+
+		FSimpleDelegateGraphTask::CreateAndDispatchWhenReady(
+			FSimpleDelegateGraphTask::FDelegate::CreateLambda([=](){
+
+				OnInAppPurchaseRestoreComplete_Delayed();
+
+			}),
+			GET_STATID(STAT_FSimpleDelegateGraphTask_DelayInAppPurchaseRestoreComplete), 
+			nullptr, 
+			ENamedThreads::GameThread
+		);
+    }
+}
+
+
+void UInAppPurchaseRestoreCallbackProxy2::OnInAppPurchaseRestoreComplete_Delayed()
+{    
+	if (bWasSuccessful)
+	{
+		OnSuccess.Broadcast(SavedPurchaseStatus, SavedReceipts);
+	}
+	else
+	{
+		OnFailure.Broadcast(SavedPurchaseStatus, SavedReceipts);
+	}
+}
+
+void UInAppPurchaseRestoreCallbackProxy2::BeginDestroy()
+{
+	Super::BeginDestroy();
+}
+
+
+UInAppPurchaseRestoreCallbackProxy2* UInAppPurchaseRestoreCallbackProxy2::CreateProxyObjectForInAppPurchaseRestore(const TArray<FInAppPurchaseProductRequest2>& ConsumableProductFlags, class APlayerController* PlayerController)
+{
+	UInAppPurchaseRestoreCallbackProxy2* Proxy = NewObject<UInAppPurchaseRestoreCallbackProxy2>();
+	Proxy->SetFlags(RF_StrongRefOnFrame);
+	Proxy->Trigger(ConsumableProductFlags, PlayerController);
+	return Proxy;
+}
+
+EInAppPurchaseStatus UInAppPurchaseRestoreCallbackProxy2::PurchaseStatusFromOnlineError(const FOnlineError& OnlineError)
+{
+	if (OnlineError.bSucceeded)
+	{
+		return EInAppPurchaseStatus::Purchased;
+	}
+	else if (OnlineError.ErrorCode.Equals(TEXT("com.epicgames.purchase.failure")))
+	{
+		return EInAppPurchaseStatus::Failed;
+	}
+	else if (OnlineError.ErrorCode.Equals(TEXT("com.epicgames.catalog_helper.user_cancelled")))
+	{
+		return EInAppPurchaseStatus::Canceled;
+	}
+	else if (OnlineError.ErrorCode.Equals(TEXT("com.epicgames.purchase.deferred")))
+	{
+		return EInAppPurchaseStatus::Deferred;
+	}
+	else if (OnlineError.ErrorCode.Equals(TEXT("com.epicgames.purchase.invalid")))
+	{
+		return EInAppPurchaseStatus::Invalid;
+	}
+
+	return EInAppPurchaseStatus::Failed;
+}
\ No newline at end of file
