Missing Include for UniqueNetIdRepl

diff --git a/Engine/Plugins/Online/OnlineSubsystemUtils/Source/OnlineSubsystemUtils/Classes/InAppPurchaseRestoreCallbackProxy2.h b/Engine/Plugins/Online/OnlineSubsystemUtils/Source/OnlineSubsystemUtils/Classes/InAppPurchaseRestoreCallbackProxy2.h
--- a/Engine/Plugins/Online/OnlineSubsystemUtils/Source/OnlineSubsystemUtils/Classes/InAppPurchaseRestoreCallbackProxy2.h
+++ b/Engine/Plugins/Online/OnlineSubsystemUtils/Source/OnlineSubsystemUtils/Classes/InAppPurchaseRestoreCallbackProxy2.h
@@ -7,6 +7,7 @@
 #include "UObject/ScriptMacros.h"
 #include "Interfaces/OnlinePurchaseInterface.h"
 #include "OnlineSubsystem.h"
+#include "GameFramework/OnlineReplStructs.h"
 #include "InAppPurchaseCallbackProxy2.h"
 #include "InAppPurchaseRestoreCallbackProxy2.generated.h"
 
