Fixup of copyright notices

diff --git a/Engine/Plugins/Online/OnlineSubsystemUtils/Source/OnlineBlueprintSupport/Classes/K2Node_InAppPurchase2.h b/Engine/Plugins/Online/OnlineSubsystemUtils/Source/OnlineBlueprintSupport/Classes/K2Node_InAppPurchase2.h
--- a/Engine/Plugins/Online/OnlineSubsystemUtils/Source/OnlineBlueprintSupport/Classes/K2Node_InAppPurchase2.h
+++ b/Engine/Plugins/Online/OnlineSubsystemUtils/Source/OnlineBlueprintSupport/Classes/K2Node_InAppPurchase2.h
@@ -1,4 +1,4 @@
-// Copyright 1998-2019 Epic Games, Inc. All Rights Reserved.
+// Copyright Epic Games, Inc. All Rights Reserved.
 
 #pragma once
 
diff --git a/Engine/Plugins/Online/OnlineSubsystemUtils/Source/OnlineBlueprintSupport/Classes/K2Node_InAppPurchaseQuery2.h b/Engine/Plugins/Online/OnlineSubsystemUtils/Source/OnlineBlueprintSupport/Classes/K2Node_InAppPurchaseQuery2.h
--- a/Engine/Plugins/Online/OnlineSubsystemUtils/Source/OnlineBlueprintSupport/Classes/K2Node_InAppPurchaseQuery2.h
+++ b/Engine/Plugins/Online/OnlineSubsystemUtils/Source/OnlineBlueprintSupport/Classes/K2Node_InAppPurchaseQuery2.h
@@ -1,4 +1,4 @@
-// Copyright 1998-2019 Epic Games, Inc. All Rights Reserved.
+// Copyright Epic Games, Inc. All Rights Reserved.
 
 #pragma once
 
diff --git a/Engine/Plugins/Online/OnlineSubsystemUtils/Source/OnlineBlueprintSupport/Classes/K2Node_InAppPurchaseQueryOwned2.h b/Engine/Plugins/Online/OnlineSubsystemUtils/Source/OnlineBlueprintSupport/Classes/K2Node_InAppPurchaseQueryOwned2.h
--- a/Engine/Plugins/Online/OnlineSubsystemUtils/Source/OnlineBlueprintSupport/Classes/K2Node_InAppPurchaseQueryOwned2.h
+++ b/Engine/Plugins/Online/OnlineSubsystemUtils/Source/OnlineBlueprintSupport/Classes/K2Node_InAppPurchaseQueryOwned2.h
@@ -1,4 +1,4 @@
-// Copyright 1998-2019 Epic Games, Inc. All Rights Reserved.
+// Copyright Epic Games, Inc. All Rights Reserved.
 
 #pragma once
 
diff --git a/Engine/Plugins/Online/OnlineSubsystemUtils/Source/OnlineBlueprintSupport/Classes/K2Node_InAppPurchaseRestore2.h b/Engine/Plugins/Online/OnlineSubsystemUtils/Source/OnlineBlueprintSupport/Classes/K2Node_InAppPurchaseRestore2.h
--- a/Engine/Plugins/Online/OnlineSubsystemUtils/Source/OnlineBlueprintSupport/Classes/K2Node_InAppPurchaseRestore2.h
+++ b/Engine/Plugins/Online/OnlineSubsystemUtils/Source/OnlineBlueprintSupport/Classes/K2Node_InAppPurchaseRestore2.h
@@ -1,4 +1,4 @@
-// Copyright 1998-2019 Epic Games, Inc. All Rights Reserved.
+// Copyright Epic Games, Inc. All Rights Reserved.
 
 #pragma once
 
diff --git a/Engine/Plugins/Online/OnlineSubsystemUtils/Source/OnlineBlueprintSupport/Classes/K2Node_InAppPurchaseUnprocessed2.h b/Engine/Plugins/Online/OnlineSubsystemUtils/Source/OnlineBlueprintSupport/Classes/K2Node_InAppPurchaseUnprocessed2.h
--- a/Engine/Plugins/Online/OnlineSubsystemUtils/Source/OnlineBlueprintSupport/Classes/K2Node_InAppPurchaseUnprocessed2.h
+++ b/Engine/Plugins/Online/OnlineSubsystemUtils/Source/OnlineBlueprintSupport/Classes/K2Node_InAppPurchaseUnprocessed2.h
@@ -1,4 +1,4 @@
-// Copyright 1998-2019 Epic Games, Inc. All Rights Reserved.
+// Copyright Epic Games, Inc. All Rights Reserved.
 
 #pragma once
 
diff --git a/Engine/Plugins/Online/OnlineSubsystemUtils/Source/OnlineBlueprintSupport/Private/K2Node_InAppPurchase2.cpp b/Engine/Plugins/Online/OnlineSubsystemUtils/Source/OnlineBlueprintSupport/Private/K2Node_InAppPurchase2.cpp
--- a/Engine/Plugins/Online/OnlineSubsystemUtils/Source/OnlineBlueprintSupport/Private/K2Node_InAppPurchase2.cpp
+++ b/Engine/Plugins/Online/OnlineSubsystemUtils/Source/OnlineBlueprintSupport/Private/K2Node_InAppPurchase2.cpp
@@ -1,4 +1,4 @@
-// Copyright 1998-2019 Epic Games, Inc. All Rights Reserved.
+// Copyright Epic Games, Inc. All Rights Reserved.
 
 #include "K2Node_InAppPurchase2.h"
 #include "InAppPurchaseCallbackProxy2.h"
diff --git a/Engine/Plugins/Online/OnlineSubsystemUtils/Source/OnlineBlueprintSupport/Private/K2Node_InAppPurchaseQuery2.cpp b/Engine/Plugins/Online/OnlineSubsystemUtils/Source/OnlineBlueprintSupport/Private/K2Node_InAppPurchaseQuery2.cpp
--- a/Engine/Plugins/Online/OnlineSubsystemUtils/Source/OnlineBlueprintSupport/Private/K2Node_InAppPurchaseQuery2.cpp
+++ b/Engine/Plugins/Online/OnlineSubsystemUtils/Source/OnlineBlueprintSupport/Private/K2Node_InAppPurchaseQuery2.cpp
@@ -1,4 +1,4 @@
-// Copyright 1998-2019 Epic Games, Inc. All Rights Reserved.
+// Copyright Epic Games, Inc. All Rights Reserved.
 
 #include "K2Node_InAppPurchaseQuery2.h"
 #include "InAppPurchaseQueryCallbackProxy2.h"
diff --git a/Engine/Plugins/Online/OnlineSubsystemUtils/Source/OnlineBlueprintSupport/Private/K2Node_InAppPurchaseQueryOwned2.cpp b/Engine/Plugins/Online/OnlineSubsystemUtils/Source/OnlineBlueprintSupport/Private/K2Node_InAppPurchaseQueryOwned2.cpp
--- a/Engine/Plugins/Online/OnlineSubsystemUtils/Source/OnlineBlueprintSupport/Private/K2Node_InAppPurchaseQueryOwned2.cpp
+++ b/Engine/Plugins/Online/OnlineSubsystemUtils/Source/OnlineBlueprintSupport/Private/K2Node_InAppPurchaseQueryOwned2.cpp
@@ -1,4 +1,4 @@
-// Copyright 1998-2019 Epic Games, Inc. All Rights Reserved.
+// Copyright Epic Games, Inc. All Rights Reserved.
 
 #include "K2Node_InAppPurchaseQueryOwned2.h"
 #include "InAppPurchaseCallbackProxy2.h"
diff --git a/Engine/Plugins/Online/OnlineSubsystemUtils/Source/OnlineBlueprintSupport/Private/K2Node_InAppPurchaseRestore2.cpp b/Engine/Plugins/Online/OnlineSubsystemUtils/Source/OnlineBlueprintSupport/Private/K2Node_InAppPurchaseRestore2.cpp
--- a/Engine/Plugins/Online/OnlineSubsystemUtils/Source/OnlineBlueprintSupport/Private/K2Node_InAppPurchaseRestore2.cpp
+++ b/Engine/Plugins/Online/OnlineSubsystemUtils/Source/OnlineBlueprintSupport/Private/K2Node_InAppPurchaseRestore2.cpp
@@ -1,4 +1,4 @@
-// Copyright 1998-2019 Epic Games, Inc. All Rights Reserved.
+// Copyright Epic Games, Inc. All Rights Reserved.
 
 #include "K2Node_InAppPurchaseRestore2.h"
 #include "InAppPurchaseRestoreCallbackProxy2.h"
diff --git a/Engine/Plugins/Online/OnlineSubsystemUtils/Source/OnlineSubsystemUtils/Classes/InAppPurchaseCallbackProxy2.h b/Engine/Plugins/Online/OnlineSubsystemUtils/Source/OnlineSubsystemUtils/Classes/InAppPurchaseCallbackProxy2.h
--- a/Engine/Plugins/Online/OnlineSubsystemUtils/Source/OnlineSubsystemUtils/Classes/InAppPurchaseCallbackProxy2.h
+++ b/Engine/Plugins/Online/OnlineSubsystemUtils/Source/OnlineSubsystemUtils/Classes/InAppPurchaseCallbackProxy2.h
@@ -1,4 +1,4 @@
-// Copyright 1998-2019 Epic Games, Inc. All Rights Reserved.
+// Copyright Epic Games, Inc. All Rights Reserved.
 #pragma once
 
 #include "CoreMinimal.h"
diff --git a/Engine/Plugins/Online/OnlineSubsystemUtils/Source/OnlineSubsystemUtils/Classes/InAppPurchaseQueryCallbackProxy2.h b/Engine/Plugins/Online/OnlineSubsystemUtils/Source/OnlineSubsystemUtils/Classes/InAppPurchaseQueryCallbackProxy2.h
--- a/Engine/Plugins/Online/OnlineSubsystemUtils/Source/OnlineSubsystemUtils/Classes/InAppPurchaseQueryCallbackProxy2.h
+++ b/Engine/Plugins/Online/OnlineSubsystemUtils/Source/OnlineSubsystemUtils/Classes/InAppPurchaseQueryCallbackProxy2.h
@@ -1,4 +1,4 @@
-// Copyright 1998-2019 Epic Games, Inc. All Rights Reserved.
+// Copyright Epic Games, Inc. All Rights Reserved.
 #pragma once
 
 #include "CoreMinimal.h"
diff --git a/Engine/Plugins/Online/OnlineSubsystemUtils/Source/OnlineSubsystemUtils/Classes/InAppPurchaseRestoreCallbackProxy2.h b/Engine/Plugins/Online/OnlineSubsystemUtils/Source/OnlineSubsystemUtils/Classes/InAppPurchaseRestoreCallbackProxy2.h
--- a/Engine/Plugins/Online/OnlineSubsystemUtils/Source/OnlineSubsystemUtils/Classes/InAppPurchaseRestoreCallbackProxy2.h
+++ b/Engine/Plugins/Online/OnlineSubsystemUtils/Source/OnlineSubsystemUtils/Classes/InAppPurchaseRestoreCallbackProxy2.h
@@ -1,4 +1,4 @@
-// Copyright 1998-2019 Epic Games, Inc. All Rights Reserved.
+// Copyright Epic Games, Inc. All Rights Reserved.
 #pragma once
 
 #include "CoreMinimal.h"
diff --git a/Engine/Plugins/Online/OnlineSubsystemUtils/Source/OnlineSubsystemUtils/Private/InAppPurchaseCallbackProxy2.cpp b/Engine/Plugins/Online/OnlineSubsystemUtils/Source/OnlineSubsystemUtils/Private/InAppPurchaseCallbackProxy2.cpp
--- a/Engine/Plugins/Online/OnlineSubsystemUtils/Source/OnlineSubsystemUtils/Private/InAppPurchaseCallbackProxy2.cpp
+++ b/Engine/Plugins/Online/OnlineSubsystemUtils/Source/OnlineSubsystemUtils/Private/InAppPurchaseCallbackProxy2.cpp
@@ -1,4 +1,4 @@
-// Copyright 1998-2019 Epic Games, Inc. All Rights Reserved.
+// Copyright Epic Games, Inc. All Rights Reserved.
 
 #include "InAppPurchaseCallbackProxy2.h"
 #include "Async/TaskGraphInterfaces.h"
diff --git a/Engine/Plugins/Online/OnlineSubsystemUtils/Source/OnlineSubsystemUtils/Private/InAppPurchaseQueryCallbackProxy2.cpp b/Engine/Plugins/Online/OnlineSubsystemUtils/Source/OnlineSubsystemUtils/Private/InAppPurchaseQueryCallbackProxy2.cpp
--- a/Engine/Plugins/Online/OnlineSubsystemUtils/Source/OnlineSubsystemUtils/Private/InAppPurchaseQueryCallbackProxy2.cpp
+++ b/Engine/Plugins/Online/OnlineSubsystemUtils/Source/OnlineSubsystemUtils/Private/InAppPurchaseQueryCallbackProxy2.cpp
@@ -1,4 +1,4 @@
-// Copyright 1998-2019 Epic Games, Inc. All Rights Reserved.
+// Copyright Epic Games, Inc. All Rights Reserved.
 
 #include "InAppPurchaseQueryCallbackProxy2.h"
 #include "Async/TaskGraphInterfaces.h"
diff --git a/Engine/Plugins/Online/OnlineSubsystemUtils/Source/OnlineSubsystemUtils/Private/InAppPurchaseRestoreCallbackProxy2.cpp b/Engine/Plugins/Online/OnlineSubsystemUtils/Source/OnlineSubsystemUtils/Private/InAppPurchaseRestoreCallbackProxy2.cpp
--- a/Engine/Plugins/Online/OnlineSubsystemUtils/Source/OnlineSubsystemUtils/Private/InAppPurchaseRestoreCallbackProxy2.cpp
+++ b/Engine/Plugins/Online/OnlineSubsystemUtils/Source/OnlineSubsystemUtils/Private/InAppPurchaseRestoreCallbackProxy2.cpp
@@ -1,4 +1,4 @@
-// Copyright 1998-2019 Epic Games, Inc. All Rights Reserved.
+// Copyright Epic Games, Inc. All Rights Reserved.
 
 #include "InAppPurchaseRestoreCallbackProxy2.h"
 #include "Async/TaskGraphInterfaces.h"
