Modified internal calls to deprecated funtions to not trigger warnings.

diff --git a/Engine/Plugins/Online/Android/OnlineSubsystemGameCircle/Source/OnlineSubsystemGameCircle/Private/OnlineSubsystemGameCircle.cpp b/Engine/Plugins/Online/Android/OnlineSubsystemGameCircle/Source/OnlineSubsystemGameCircle/Private/OnlineSubsystemGameCircle.cpp
--- a/Engine/Plugins/Online/Android/OnlineSubsystemGameCircle/Source/OnlineSubsystemGameCircle/Private/OnlineSubsystemGameCircle.cpp
+++ b/Engine/Plugins/Online/Android/OnlineSubsystemGameCircle/Source/OnlineSubsystemGameCircle/Private/OnlineSubsystemGameCircle.cpp
@@ -19,10 +19,12 @@ IOnlineIdentityPtr FOnlineSubsystemGameCircle::GetIdentityInterface() const
 	return IdentityInterface;
 }
 
+PRAGMA_DISABLE_DEPRECATION_WARNINGS
 IOnlineStorePtr FOnlineSubsystemGameCircle::GetStoreInterface() const
 {
 	return StoreInterface;
 }
+PRAGMA_ENABLE_DEPRECATION_WARNINGS
 
 IOnlineSessionPtr FOnlineSubsystemGameCircle::GetSessionInterface() const
 {
diff --git a/Engine/Plugins/Online/Android/OnlineSubsystemGameCircle/Source/OnlineSubsystemGameCircle/Public/OnlineSubsystemGameCircle.h b/Engine/Plugins/Online/Android/OnlineSubsystemGameCircle/Source/OnlineSubsystemGameCircle/Public/OnlineSubsystemGameCircle.h
--- a/Engine/Plugins/Online/Android/OnlineSubsystemGameCircle/Source/OnlineSubsystemGameCircle/Public/OnlineSubsystemGameCircle.h
+++ b/Engine/Plugins/Online/Android/OnlineSubsystemGameCircle/Source/OnlineSubsystemGameCircle/Public/OnlineSubsystemGameCircle.h
@@ -40,7 +40,9 @@ public:
 	virtual IOnlineIdentityPtr GetIdentityInterface() const override;
 	virtual IOnlineTitleFilePtr GetTitleFileInterface() const override;
 	virtual IOnlineEntitlementsPtr GetEntitlementsInterface() const override;
+	PRAGMA_DISABLE_DEPRECATION_WARNINGS
 	virtual IOnlineStorePtr GetStoreInterface() const override;
+	PRAGMA_ENABLE_DEPRECATION_WARNINGS
 	virtual IOnlineStoreV2Ptr GetStoreV2Interface() const override { return nullptr; }
 	virtual IOnlinePurchasePtr GetPurchaseInterface() const override { return nullptr; }
 	virtual IOnlineEventsPtr GetEventsInterface() const override { return nullptr; }
diff --git a/Engine/Plugins/Online/Android/OnlineSubsystemGooglePlay/Source/Private/OnlineSubsystemGooglePlay.cpp b/Engine/Plugins/Online/Android/OnlineSubsystemGooglePlay/Source/Private/OnlineSubsystemGooglePlay.cpp
--- a/Engine/Plugins/Online/Android/OnlineSubsystemGooglePlay/Source/Private/OnlineSubsystemGooglePlay.cpp
+++ b/Engine/Plugins/Online/Android/OnlineSubsystemGooglePlay/Source/Private/OnlineSubsystemGooglePlay.cpp
@@ -47,10 +47,12 @@ IOnlineIdentityPtr FOnlineSubsystemGooglePlay::GetIdentityInterface() const
 	return IdentityInterface;
 }
 
+PRAGMA_DISABLE_DEPRECATION_WARNINGS
 IOnlineStorePtr FOnlineSubsystemGooglePlay::GetStoreInterface() const
 {
 	return StoreInterface;
 }
+PRAGMA_ENABLE_DEPRECATION_WARNINGS
 
 IOnlineStoreV2Ptr FOnlineSubsystemGooglePlay::GetStoreV2Interface() const
 {
diff --git a/Engine/Plugins/Online/Android/OnlineSubsystemGooglePlay/Source/Public/OnlineSubsystemGooglePlay.h b/Engine/Plugins/Online/Android/OnlineSubsystemGooglePlay/Source/Public/OnlineSubsystemGooglePlay.h
--- a/Engine/Plugins/Online/Android/OnlineSubsystemGooglePlay/Source/Public/OnlineSubsystemGooglePlay.h
+++ b/Engine/Plugins/Online/Android/OnlineSubsystemGooglePlay/Source/Public/OnlineSubsystemGooglePlay.h
@@ -58,7 +58,9 @@ public:
 	virtual IOnlineIdentityPtr GetIdentityInterface() const override;
 	virtual IOnlineTitleFilePtr GetTitleFileInterface() const override;
 	virtual IOnlineEntitlementsPtr GetEntitlementsInterface() const override;
+	PRAGMA_DISABLE_DEPRECATION_WARNINGS
 	virtual IOnlineStorePtr GetStoreInterface() const override;
+	PRAGMA_ENABLE_DEPRECATION_WARNINGS
 	virtual IOnlineStoreV2Ptr GetStoreV2Interface() const override;
 	virtual IOnlinePurchasePtr GetPurchaseInterface() const override;
 	virtual IOnlineEventsPtr GetEventsInterface() const override { return nullptr; }
diff --git a/Engine/Plugins/Online/IOS/OnlineSubsystemIOS/Source/Private/OnlineSubsystemIOS.cpp b/Engine/Plugins/Online/IOS/OnlineSubsystemIOS/Source/Private/OnlineSubsystemIOS.cpp
--- a/Engine/Plugins/Online/IOS/OnlineSubsystemIOS/Source/Private/OnlineSubsystemIOS.cpp
+++ b/Engine/Plugins/Online/IOS/OnlineSubsystemIOS/Source/Private/OnlineSubsystemIOS.cpp
@@ -79,10 +79,12 @@ IOnlineEntitlementsPtr FOnlineSubsystemIOS::GetEntitlementsInterface() const
 	return nullptr;
 }
 
+PRAGMA_DISABLE_DEPRECATION_WARNINGS
 IOnlineStorePtr FOnlineSubsystemIOS::GetStoreInterface() const
 {
 	return StoreInterface;
 }
+PRAGMA_ENABLE_DEPRECATION_WARNINGS
 
 IOnlineStoreV2Ptr FOnlineSubsystemIOS::GetStoreV2Interface() const
 {
diff --git a/Engine/Plugins/Online/IOS/OnlineSubsystemIOS/Source/Public/OnlineSubsystemIOS.h b/Engine/Plugins/Online/IOS/OnlineSubsystemIOS/Source/Public/OnlineSubsystemIOS.h
--- a/Engine/Plugins/Online/IOS/OnlineSubsystemIOS/Source/Public/OnlineSubsystemIOS.h
+++ b/Engine/Plugins/Online/IOS/OnlineSubsystemIOS/Source/Public/OnlineSubsystemIOS.h
@@ -47,7 +47,9 @@ public:
 	virtual IOnlineIdentityPtr GetIdentityInterface() const override;
 	virtual IOnlineTitleFilePtr GetTitleFileInterface() const override;
 	virtual IOnlineEntitlementsPtr GetEntitlementsInterface() const override;
+	PRAGMA_DISABLE_DEPRECATION_WARNINGS
 	virtual IOnlineStorePtr GetStoreInterface() const override;
+	PRAGMA_ENABLE_DEPRECATION_WARNINGS
 	virtual IOnlineStoreV2Ptr GetStoreV2Interface() const override;
 	virtual IOnlinePurchasePtr GetPurchaseInterface() const override;
 	virtual IOnlineEventsPtr GetEventsInterface() const override;
diff --git a/Engine/Plugins/Online/OnlineSubsystem/Source/Public/OnlineSubsystem.h b/Engine/Plugins/Online/OnlineSubsystem/Source/Public/OnlineSubsystem.h
--- a/Engine/Plugins/Online/OnlineSubsystem/Source/Public/OnlineSubsystem.h
+++ b/Engine/Plugins/Online/OnlineSubsystem/Source/Public/OnlineSubsystem.h
@@ -370,8 +370,8 @@ public:
 	 * Get the interface for accessing an online store
 	 * @return Interface pointer for the appropriate online store service
 	 */
-	UE_DEPRECATED(4.24, "Move to OnlineStoreInterfaceV2 and OnlinePurchaseInterface.")
-	virtual IOnlineStorePtr GetStoreInterface() const = 0;
+	UE_DEPRECATED(4.25, "Move to OnlineStoreInterfaceV2 and OnlinePurchaseInterface.")
+	virtual IOnlineStorePtr GetStoreInterface() const { return nullptr; }
 
 	/** 
 	 * Get the interface for accessing an online store
diff --git a/Engine/Plugins/Online/OnlineSubsystemAmazon/Source/Public/OnlineSubsystemAmazon.h b/Engine/Plugins/Online/OnlineSubsystemAmazon/Source/Public/OnlineSubsystemAmazon.h
--- a/Engine/Plugins/Online/OnlineSubsystemAmazon/Source/Public/OnlineSubsystemAmazon.h
+++ b/Engine/Plugins/Online/OnlineSubsystemAmazon/Source/Public/OnlineSubsystemAmazon.h
@@ -47,7 +47,6 @@ public:
 	virtual IOnlineTitleFilePtr GetTitleFileInterface() const override { return nullptr; }
 	virtual IOnlineEntitlementsPtr GetEntitlementsInterface() const override { return nullptr; }
 	virtual IOnlineIdentityPtr GetIdentityInterface() const override;
-	virtual IOnlineStorePtr GetStoreInterface() const override { return nullptr; }
 	virtual IOnlineStoreV2Ptr GetStoreV2Interface() const override { return nullptr; }
 	virtual IOnlinePurchasePtr GetPurchaseInterface() const override { return nullptr; }
 	virtual IOnlineEventsPtr GetEventsInterface() const override { return nullptr; }
diff --git a/Engine/Plugins/Online/OnlineSubsystemFacebook/Source/Private/OnlineSubsystemFacebookCommon.cpp b/Engine/Plugins/Online/OnlineSubsystemFacebook/Source/Private/OnlineSubsystemFacebookCommon.cpp
--- a/Engine/Plugins/Online/OnlineSubsystemFacebook/Source/Private/OnlineSubsystemFacebookCommon.cpp
+++ b/Engine/Plugins/Online/OnlineSubsystemFacebook/Source/Private/OnlineSubsystemFacebookCommon.cpp
@@ -156,11 +156,6 @@ IOnlineEntitlementsPtr FOnlineSubsystemFacebookCommon::GetEntitlementsInterface(
 	return nullptr;
 }
 
-IOnlineStorePtr FOnlineSubsystemFacebookCommon::GetStoreInterface() const
-{
-	return nullptr;
-}
-
 IOnlineEventsPtr FOnlineSubsystemFacebookCommon::GetEventsInterface() const
 {
 	return nullptr;
diff --git a/Engine/Plugins/Online/OnlineSubsystemFacebook/Source/Private/OnlineSubsystemFacebookCommon.h b/Engine/Plugins/Online/OnlineSubsystemFacebook/Source/Private/OnlineSubsystemFacebookCommon.h
--- a/Engine/Plugins/Online/OnlineSubsystemFacebook/Source/Private/OnlineSubsystemFacebookCommon.h
+++ b/Engine/Plugins/Online/OnlineSubsystemFacebook/Source/Private/OnlineSubsystemFacebookCommon.h
@@ -36,7 +36,6 @@ public:
 	virtual IOnlineIdentityPtr GetIdentityInterface() const override;
 	virtual IOnlineTitleFilePtr GetTitleFileInterface() const override;
 	virtual IOnlineEntitlementsPtr GetEntitlementsInterface() const override;
-	virtual IOnlineStorePtr GetStoreInterface() const override;
 	virtual IOnlineStoreV2Ptr GetStoreV2Interface() const override { return nullptr; }
 	virtual IOnlinePurchasePtr GetPurchaseInterface() const override { return nullptr; }
 	virtual IOnlineEventsPtr GetEventsInterface() const override;
diff --git a/Engine/Plugins/Online/OnlineSubsystemGoogle/Source/Private/OnlineSubsystemGoogleCommon.cpp b/Engine/Plugins/Online/OnlineSubsystemGoogle/Source/Private/OnlineSubsystemGoogleCommon.cpp
--- a/Engine/Plugins/Online/OnlineSubsystemGoogle/Source/Private/OnlineSubsystemGoogleCommon.cpp
+++ b/Engine/Plugins/Online/OnlineSubsystemGoogle/Source/Private/OnlineSubsystemGoogleCommon.cpp
@@ -180,11 +180,6 @@ IOnlineEntitlementsPtr FOnlineSubsystemGoogleCommon::GetEntitlementsInterface()
 	return nullptr;
 }
 
-IOnlineStorePtr FOnlineSubsystemGoogleCommon::GetStoreInterface() const
-{
-	return nullptr;
-}
-
 IOnlineEventsPtr FOnlineSubsystemGoogleCommon::GetEventsInterface() const
 {
 	return nullptr;
diff --git a/Engine/Plugins/Online/OnlineSubsystemGoogle/Source/Public/OnlineSubsystemGoogleCommon.h b/Engine/Plugins/Online/OnlineSubsystemGoogle/Source/Public/OnlineSubsystemGoogleCommon.h
--- a/Engine/Plugins/Online/OnlineSubsystemGoogle/Source/Public/OnlineSubsystemGoogleCommon.h
+++ b/Engine/Plugins/Online/OnlineSubsystemGoogle/Source/Public/OnlineSubsystemGoogleCommon.h
@@ -49,7 +49,6 @@ public:
 	virtual IOnlineIdentityPtr GetIdentityInterface() const override;
 	virtual IOnlineTitleFilePtr GetTitleFileInterface() const override;
 	virtual IOnlineEntitlementsPtr GetEntitlementsInterface() const override;
-	virtual IOnlineStorePtr GetStoreInterface() const override;
 	virtual IOnlineStoreV2Ptr GetStoreV2Interface() const override { return nullptr; }
 	virtual IOnlinePurchasePtr GetPurchaseInterface() const override { return nullptr; }
 	virtual IOnlineEventsPtr GetEventsInterface() const override;
diff --git a/Engine/Plugins/Online/OnlineSubsystemNull/Source/Private/OnlineSubsystemNull.cpp b/Engine/Plugins/Online/OnlineSubsystemNull/Source/Private/OnlineSubsystemNull.cpp
--- a/Engine/Plugins/Online/OnlineSubsystemNull/Source/Private/OnlineSubsystemNull.cpp
+++ b/Engine/Plugins/Online/OnlineSubsystemNull/Source/Private/OnlineSubsystemNull.cpp
@@ -89,11 +89,6 @@ IOnlineTitleFilePtr FOnlineSubsystemNull::GetTitleFileInterface() const
 	return nullptr;
 }
 
-IOnlineStorePtr FOnlineSubsystemNull::GetStoreInterface() const
-{
-	return nullptr;
-}
-
 IOnlineStoreV2Ptr FOnlineSubsystemNull::GetStoreV2Interface() const
 {
 	return StoreV2Interface;
diff --git a/Engine/Plugins/Online/OnlineSubsystemNull/Source/Public/OnlineSubsystemNull.h b/Engine/Plugins/Online/OnlineSubsystemNull/Source/Public/OnlineSubsystemNull.h
--- a/Engine/Plugins/Online/OnlineSubsystemNull/Source/Public/OnlineSubsystemNull.h
+++ b/Engine/Plugins/Online/OnlineSubsystemNull/Source/Public/OnlineSubsystemNull.h
@@ -52,7 +52,6 @@ public:
 	virtual IOnlineTimePtr GetTimeInterface() const override;
 	virtual IOnlineIdentityPtr GetIdentityInterface() const override;
 	virtual IOnlineTitleFilePtr GetTitleFileInterface() const override;
-	virtual IOnlineStorePtr GetStoreInterface() const override;
 	virtual IOnlineStoreV2Ptr GetStoreV2Interface() const override;
 	virtual IOnlinePurchasePtr GetPurchaseInterface() const override;
 	virtual IOnlineEventsPtr GetEventsInterface() const override;
diff --git a/Engine/Plugins/Online/OnlineSubsystemOculus/Source/Private/OnlineSubsystemOculus.cpp b/Engine/Plugins/Online/OnlineSubsystemOculus/Source/Private/OnlineSubsystemOculus.cpp
--- a/Engine/Plugins/Online/OnlineSubsystemOculus/Source/Private/OnlineSubsystemOculus.cpp
+++ b/Engine/Plugins/Online/OnlineSubsystemOculus/Source/Private/OnlineSubsystemOculus.cpp
@@ -89,11 +89,6 @@ IOnlineTitleFilePtr FOnlineSubsystemOculus::GetTitleFileInterface() const
 	return nullptr;
 }
 
-IOnlineStorePtr FOnlineSubsystemOculus::GetStoreInterface() const
-{
-	return nullptr;
-}
-
 IOnlineStoreV2Ptr FOnlineSubsystemOculus::GetStoreV2Interface() const
 {
 	return nullptr;
diff --git a/Engine/Plugins/Online/OnlineSubsystemOculus/Source/Public/OnlineSubsystemOculus.h b/Engine/Plugins/Online/OnlineSubsystemOculus/Source/Public/OnlineSubsystemOculus.h
--- a/Engine/Plugins/Online/OnlineSubsystemOculus/Source/Public/OnlineSubsystemOculus.h
+++ b/Engine/Plugins/Online/OnlineSubsystemOculus/Source/Public/OnlineSubsystemOculus.h
@@ -46,7 +46,6 @@ public:
 	virtual IOnlineTimePtr GetTimeInterface() const override;
 	virtual IOnlineIdentityPtr GetIdentityInterface() const override;
 	virtual IOnlineTitleFilePtr GetTitleFileInterface() const override;
-	virtual IOnlineStorePtr GetStoreInterface() const override;
 	virtual IOnlineStoreV2Ptr GetStoreV2Interface() const override;
 	virtual IOnlinePurchasePtr GetPurchaseInterface() const override;
 	virtual IOnlineEventsPtr GetEventsInterface() const override;
diff --git a/Engine/Plugins/Online/OnlineSubsystemSteam/Source/Private/OnlineSubsystemSteam.cpp b/Engine/Plugins/Online/OnlineSubsystemSteam/Source/Private/OnlineSubsystemSteam.cpp
--- a/Engine/Plugins/Online/OnlineSubsystemSteam/Source/Private/OnlineSubsystemSteam.cpp
+++ b/Engine/Plugins/Online/OnlineSubsystemSteam/Source/Private/OnlineSubsystemSteam.cpp
@@ -322,11 +322,6 @@ IOnlineEntitlementsPtr FOnlineSubsystemSteam::GetEntitlementsInterface() const
 	return nullptr;
 }
 
-IOnlineStorePtr FOnlineSubsystemSteam::GetStoreInterface() const
-{
-	return nullptr;
-}
-
 IOnlineEventsPtr FOnlineSubsystemSteam::GetEventsInterface() const
 {
 	return nullptr;
diff --git a/Engine/Plugins/Online/OnlineSubsystemSteam/Source/Public/OnlineSubsystemSteam.h b/Engine/Plugins/Online/OnlineSubsystemSteam/Source/Public/OnlineSubsystemSteam.h
--- a/Engine/Plugins/Online/OnlineSubsystemSteam/Source/Public/OnlineSubsystemSteam.h
+++ b/Engine/Plugins/Online/OnlineSubsystemSteam/Source/Public/OnlineSubsystemSteam.h
@@ -265,7 +265,6 @@ public:
 	virtual IOnlineIdentityPtr GetIdentityInterface() const override;
 	virtual IOnlineTitleFilePtr GetTitleFileInterface() const override;
 	virtual IOnlineEntitlementsPtr GetEntitlementsInterface() const override;
-	virtual IOnlineStorePtr GetStoreInterface() const override;
 	virtual IOnlineStoreV2Ptr GetStoreV2Interface() const override { return nullptr; }
 	virtual IOnlinePurchasePtr GetPurchaseInterface() const override { return nullptr; }
 	virtual IOnlineEventsPtr GetEventsInterface() const override;
diff --git a/Engine/Plugins/Online/OnlineSubsystemTencent/Source/Private/OnlineSubsystemTencent.cpp b/Engine/Plugins/Online/OnlineSubsystemTencent/Source/Private/OnlineSubsystemTencent.cpp
--- a/Engine/Plugins/Online/OnlineSubsystemTencent/Source/Private/OnlineSubsystemTencent.cpp
+++ b/Engine/Plugins/Online/OnlineSubsystemTencent/Source/Private/OnlineSubsystemTencent.cpp
@@ -107,11 +107,6 @@ IOnlineTitleFilePtr FOnlineSubsystemTencent:: GetTitleFileInterface() const
 	return nullptr;
 }
 
-IOnlineStorePtr FOnlineSubsystemTencent::GetStoreInterface() const
-{
-	return nullptr;
-}
-
 IOnlineStoreV2Ptr FOnlineSubsystemTencent::GetStoreV2Interface() const
 {
 #if	WITH_TENCENT_RAIL_SDK
diff --git a/Engine/Plugins/Online/OnlineSubsystemTencent/Source/Public/OnlineSubsystemTencent.h b/Engine/Plugins/Online/OnlineSubsystemTencent/Source/Public/OnlineSubsystemTencent.h
--- a/Engine/Plugins/Online/OnlineSubsystemTencent/Source/Public/OnlineSubsystemTencent.h
+++ b/Engine/Plugins/Online/OnlineSubsystemTencent/Source/Public/OnlineSubsystemTencent.h
@@ -71,7 +71,6 @@ public:
 	virtual IOnlineTimePtr GetTimeInterface() const override;
 	virtual IOnlineIdentityPtr GetIdentityInterface() const override;
 	virtual IOnlineTitleFilePtr GetTitleFileInterface() const override;
-	virtual IOnlineStorePtr GetStoreInterface() const override;
 	virtual IOnlineStoreV2Ptr GetStoreV2Interface() const override;
 	virtual IOnlinePurchasePtr GetPurchaseInterface() const override;
 	virtual IOnlineEventsPtr GetEventsInterface() const override;
diff --git a/Engine/Plugins/Online/OnlineSubsystemTwitch/Source/Private/OnlineSubsystemTwitch.cpp b/Engine/Plugins/Online/OnlineSubsystemTwitch/Source/Private/OnlineSubsystemTwitch.cpp
--- a/Engine/Plugins/Online/OnlineSubsystemTwitch/Source/Private/OnlineSubsystemTwitch.cpp
+++ b/Engine/Plugins/Online/OnlineSubsystemTwitch/Source/Private/OnlineSubsystemTwitch.cpp
@@ -71,11 +71,6 @@ IOnlineTitleFilePtr FOnlineSubsystemTwitch:: GetTitleFileInterface() const
 	return nullptr;
 }
 
-IOnlineStorePtr FOnlineSubsystemTwitch::GetStoreInterface() const
-{
-	return nullptr;
-}
-
 IOnlineStoreV2Ptr FOnlineSubsystemTwitch::GetStoreV2Interface() const
 {
 	return nullptr;
diff --git a/Engine/Plugins/Online/OnlineSubsystemTwitch/Source/Public/OnlineSubsystemTwitch.h b/Engine/Plugins/Online/OnlineSubsystemTwitch/Source/Public/OnlineSubsystemTwitch.h
--- a/Engine/Plugins/Online/OnlineSubsystemTwitch/Source/Public/OnlineSubsystemTwitch.h
+++ b/Engine/Plugins/Online/OnlineSubsystemTwitch/Source/Public/OnlineSubsystemTwitch.h
@@ -37,7 +37,6 @@ public:
 	virtual IOnlineTimePtr GetTimeInterface() const override;
 	virtual IOnlineIdentityPtr GetIdentityInterface() const override;
 	virtual IOnlineTitleFilePtr GetTitleFileInterface() const override;
-	virtual IOnlineStorePtr GetStoreInterface() const override;
 	virtual IOnlineStoreV2Ptr GetStoreV2Interface() const override;
 	virtual IOnlinePurchasePtr GetPurchaseInterface() const override;
 	virtual IOnlineEventsPtr GetEventsInterface() const override;
diff --git a/Engine/Plugins/Online/OnlineSubsystemUtils/Source/OnlineSubsystemUtils/Private/InAppPurchaseCallbackProxy.cpp b/Engine/Plugins/Online/OnlineSubsystemUtils/Source/OnlineSubsystemUtils/Private/InAppPurchaseCallbackProxy.cpp
--- a/Engine/Plugins/Online/OnlineSubsystemUtils/Source/OnlineSubsystemUtils/Private/InAppPurchaseCallbackProxy.cpp
+++ b/Engine/Plugins/Online/OnlineSubsystemUtils/Source/OnlineSubsystemUtils/Private/InAppPurchaseCallbackProxy.cpp
@@ -28,7 +28,9 @@ void UInAppPurchaseCallbackProxy::Trigger(APlayerController* PlayerController, c
 	{
 		if (IOnlineSubsystem* const OnlineSub = IOnlineSubsystem::IsLoaded() ? IOnlineSubsystem::Get() : nullptr)
 		{
+			PRAGMA_DISABLE_DEPRECATION_WARNINGS
 			IOnlineStorePtr StoreInterface = OnlineSub->GetStoreInterface();
+			PRAGMA_ENABLE_DEPRECATION_WARNINGS
 			if (StoreInterface.IsValid())
 			{
 				bFailedToEvenSubmit = false;
@@ -123,7 +125,9 @@ void UInAppPurchaseCallbackProxy::RemoveDelegate()
 	{
 		if (IOnlineSubsystem* OnlineSub = IOnlineSubsystem::IsLoaded() ? IOnlineSubsystem::Get() : nullptr)
 		{
+			PRAGMA_DISABLE_DEPRECATION_WARNINGS
 			IOnlineStorePtr InAppPurchases = OnlineSub->GetStoreInterface();
+			PRAGMA_ENABLE_DEPRECATION_WARNINGS
 			if (InAppPurchases.IsValid())
 			{
 				InAppPurchases->ClearOnInAppPurchaseCompleteDelegate_Handle(InAppPurchaseCompleteDelegateHandle);
diff --git a/Engine/Plugins/Online/OnlineSubsystemUtils/Source/OnlineSubsystemUtils/Private/InAppPurchaseQueryCallbackProxy.cpp b/Engine/Plugins/Online/OnlineSubsystemUtils/Source/OnlineSubsystemUtils/Private/InAppPurchaseQueryCallbackProxy.cpp
--- a/Engine/Plugins/Online/OnlineSubsystemUtils/Source/OnlineSubsystemUtils/Private/InAppPurchaseQueryCallbackProxy.cpp
+++ b/Engine/Plugins/Online/OnlineSubsystemUtils/Source/OnlineSubsystemUtils/Private/InAppPurchaseQueryCallbackProxy.cpp
@@ -23,7 +23,9 @@ void UInAppPurchaseQueryCallbackProxy::TriggerQuery(APlayerController* PlayerCon
 	{
 		if (IOnlineSubsystem* const OnlineSub = IOnlineSubsystem::IsLoaded() ? IOnlineSubsystem::Get() : nullptr)
 		{
+			PRAGMA_DISABLE_DEPRECATION_WARNINGS
 			IOnlineStorePtr StoreInterface = OnlineSub->GetStoreInterface();
+			PRAGMA_ENABLE_DEPRECATION_WARNINGS
 			if (StoreInterface.IsValid())
 			{
 				bFailedToEvenSubmit = false;
@@ -109,7 +111,9 @@ void UInAppPurchaseQueryCallbackProxy::RemoveDelegate()
 	{
 		if (IOnlineSubsystem* OnlineSub = IOnlineSubsystem::IsLoaded() ? IOnlineSubsystem::Get() : nullptr)
 		{
+			PRAGMA_DISABLE_DEPRECATION_WARNINGS
 			IOnlineStorePtr InAppPurchases = OnlineSub->GetStoreInterface();
+			PRAGMA_ENABLE_DEPRECATION_WARNINGS
 			if (InAppPurchases.IsValid())
 			{
 				InAppPurchases->ClearOnQueryForAvailablePurchasesCompleteDelegate_Handle(InAppPurchaseReadCompleteDelegateHandle);
diff --git a/Engine/Plugins/Online/OnlineSubsystemUtils/Source/OnlineSubsystemUtils/Private/InAppPurchaseRestoreCallbackProxy.cpp b/Engine/Plugins/Online/OnlineSubsystemUtils/Source/OnlineSubsystemUtils/Private/InAppPurchaseRestoreCallbackProxy.cpp
--- a/Engine/Plugins/Online/OnlineSubsystemUtils/Source/OnlineSubsystemUtils/Private/InAppPurchaseRestoreCallbackProxy.cpp
+++ b/Engine/Plugins/Online/OnlineSubsystemUtils/Source/OnlineSubsystemUtils/Private/InAppPurchaseRestoreCallbackProxy.cpp
@@ -25,7 +25,9 @@ void UInAppPurchaseRestoreCallbackProxy::Trigger(const TArray<FInAppPurchaseProd
 	{
 		if (IOnlineSubsystem* const OnlineSub = IOnlineSubsystem::IsLoaded() ? IOnlineSubsystem::Get() : nullptr)
 		{
+			PRAGMA_DISABLE_DEPRECATION_WARNINGS
 			IOnlineStorePtr StoreInterface = OnlineSub->GetStoreInterface();
+			PRAGMA_ENABLE_DEPRECATION_WARNINGS
 			if (StoreInterface.IsValid())
 			{
 				bFailedToEvenSubmit = false;
@@ -109,7 +111,9 @@ void UInAppPurchaseRestoreCallbackProxy::RemoveDelegate()
 	{
 		if (IOnlineSubsystem* OnlineSub = IOnlineSubsystem::IsLoaded() ? IOnlineSubsystem::Get() : nullptr)
 		{
+			PRAGMA_DISABLE_DEPRECATION_WARNINGS
 			IOnlineStorePtr InAppPurchases = OnlineSub->GetStoreInterface();
+			PRAGMA_ENABLE_DEPRECATION_WARNINGS
 			if (InAppPurchases.IsValid())
 			{
 				InAppPurchases->ClearOnInAppPurchaseRestoreCompleteDelegate_Handle(InAppPurchaseRestoreCompleteDelegateHandle);
