fixed missing include

diff --git a/Engine/Plugins/Online/OnlineSubsystemUtils/Source/OnlineSubsystemUtils/Private/InAppPurchaseCallbackProxy2.cpp b/Engine/Plugins/Online/OnlineSubsystemUtils/Source/OnlineSubsystemUtils/Private/InAppPurchaseCallbackProxy2.cpp
--- a/Engine/Plugins/Online/OnlineSubsystemUtils/Source/OnlineSubsystemUtils/Private/InAppPurchaseCallbackProxy2.cpp
+++ b/Engine/Plugins/Online/OnlineSubsystemUtils/Source/OnlineSubsystemUtils/Private/InAppPurchaseCallbackProxy2.cpp
@@ -5,6 +5,7 @@
 #include "GameFramework/PlayerController.h"
 #include "OnlineSubsystem.h"
 #include "Engine/World.h"
+#include "Engine/LocalPlayer.h"
 
 //////////////////////////////////////////////////////////////////////////
 // UInAppPurchaseCallbackProxy
diff --git a/Engine/Plugins/Online/OnlineSubsystemUtils/Source/OnlineSubsystemUtils/Private/InAppPurchaseQueryCallbackProxy2.cpp b/Engine/Plugins/Online/OnlineSubsystemUtils/Source/OnlineSubsystemUtils/Private/InAppPurchaseQueryCallbackProxy2.cpp
--- a/Engine/Plugins/Online/OnlineSubsystemUtils/Source/OnlineSubsystemUtils/Private/InAppPurchaseQueryCallbackProxy2.cpp
+++ b/Engine/Plugins/Online/OnlineSubsystemUtils/Source/OnlineSubsystemUtils/Private/InAppPurchaseQueryCallbackProxy2.cpp
@@ -5,6 +5,7 @@
 #include "GameFramework/PlayerController.h"
 #include "OnlineSubsystem.h"
 #include "Engine/World.h"
+#include "Engine/LocalPlayer.h"
 
 //////////////////////////////////////////////////////////////////////////
 // UInAppPurchaseQueryCallbackProxy
diff --git a/Engine/Plugins/Online/OnlineSubsystemUtils/Source/OnlineSubsystemUtils/Private/InAppPurchaseRestoreCallbackProxy2.cpp b/Engine/Plugins/Online/OnlineSubsystemUtils/Source/OnlineSubsystemUtils/Private/InAppPurchaseRestoreCallbackProxy2.cpp
--- a/Engine/Plugins/Online/OnlineSubsystemUtils/Source/OnlineSubsystemUtils/Private/InAppPurchaseRestoreCallbackProxy2.cpp
+++ b/Engine/Plugins/Online/OnlineSubsystemUtils/Source/OnlineSubsystemUtils/Private/InAppPurchaseRestoreCallbackProxy2.cpp
@@ -5,6 +5,7 @@
 #include "GameFramework/PlayerController.h"
 #include "OnlineSubsystem.h"
 #include "Engine/World.h"
+#include "Engine/LocalPlayer.h"
 
 //////////////////////////////////////////////////////////////////////////
 // UInAppPurchaseRestoreCallbackProxy
