Fix shootergame usage of online subsystem to work properly with multiple PIE worlds

diff --git a/Engine/Plugins/Online/OnlineSubsystemUtils/Source/OnlineSubsystemUtils/Public/OnlineSubsystemUtils.h b/Engine/Plugins/Online/OnlineSubsystemUtils/Source/OnlineSubsystemUtils/Public/OnlineSubsystemUtils.h
--- a/Engine/Plugins/Online/OnlineSubsystemUtils/Source/OnlineSubsystemUtils/Public/OnlineSubsystemUtils.h
+++ b/Engine/Plugins/Online/OnlineSubsystemUtils/Source/OnlineSubsystemUtils/Public/OnlineSubsystemUtils.h
@@ -294,6 +294,8 @@ namespace Online
 	IMPLEMENT_GET_INTERFACE(Leaderboards);
 	IMPLEMENT_GET_INTERFACE(Achievements);
 	IMPLEMENT_GET_INTERFACE(Presence);
+	IMPLEMENT_GET_INTERFACE(Events);
+	IMPLEMENT_GET_INTERFACE(Stats);
 }
 
 #undef IMPLEMENT_GET_INTERFACE
