[NO NON-PII COMMIT MESSAGE]

diff --git a/Engine/Plugins/Online/OnlineSubsystemUtils/Source/OnlineSubsystemUtils/Private/VoiceInterfaceImpl.cpp b/Engine/Plugins/Online/OnlineSubsystemUtils/Source/OnlineSubsystemUtils/Private/VoiceInterfaceImpl.cpp
--- a/Engine/Plugins/Online/OnlineSubsystemUtils/Source/OnlineSubsystemUtils/Private/VoiceInterfaceImpl.cpp
+++ b/Engine/Plugins/Online/OnlineSubsystemUtils/Source/OnlineSubsystemUtils/Private/VoiceInterfaceImpl.cpp
@@ -560,12 +560,8 @@ void FOnlineVoiceImpl::ProcessMuteChangeNotification()
 				ULocalPlayer* LP = GEngine->FindFirstLocalPlayerFromControllerId(Index);
 				if (LP && LP->PlayerController)
 				{
-					// If there is a player controller, we can mute/unmute people
-					if (LocalTalkers[Index].bIsRegistered)
-					{
-						// Use the common method of checking muting
-						UpdateMuteListForLocalTalker(Index, LP->PlayerController);
-					}
+					// Use the common method of checking muting
+					UpdateMuteListForLocalTalker(Index, LP->PlayerController);
 				}
 			}
 		}
