Implement showing and sending friend messages on Steam. Additionally implements the store overlay and functionality for automatically adding a product to the user's cart.

Also updates test cases to existing external UI harness to allow for the new functionality.

diff --git a/Engine/Plugins/Online/OnlineSubsystem/Source/Public/Interfaces/OnlineExternalUIInterface.h b/Engine/Plugins/Online/OnlineSubsystem/Source/Public/Interfaces/OnlineExternalUIInterface.h
--- a/Engine/Plugins/Online/OnlineSubsystem/Source/Public/Interfaces/OnlineExternalUIInterface.h
+++ b/Engine/Plugins/Online/OnlineSubsystem/Source/Public/Interfaces/OnlineExternalUIInterface.h
@@ -232,11 +232,15 @@ struct FShowStoreParams
 	/** Product to show directly instead of the whole store */
 	FUniqueOfferId ProductId;
 
+	/** If the specific product id should be added to the user's cart automatically */
+	bool bAddToCart;
+
 	/**
 	 * Constructor
 	 */
 	FShowStoreParams(const FString& InCategory = FString())
 		: Category(InCategory)
+		, bAddToCart(false)
 	{}
 };
 
@@ -409,6 +413,23 @@ public:
 	 */
 	virtual bool ShowSendMessageUI(int32 LocalUserNum, const FShowSendMessageParams& ShowParams, const FOnShowSendMessageUIClosedDelegate& Delegate = FOnShowSendMessageUIClosedDelegate()) = 0;
 
+	/**
+	 * Display the platform UI for sending a mailbox message to the specified user.
+	 * Will attempt to automatically send the message if the platform allows for it.
+	 *
+	 * @param LocalUserNum the controller number of the associated user
+	 * @param Recipient the user that should receive this message
+	 * @param ShowParams configuration for the UI display
+	 *
+	 * @return true if it was able to send the message and show the UI, false if it failed
+	 */
+	virtual bool ShowSendMessageToUserUI(int32 LocalUserNum, const FUniqueNetId& Recipient, const FShowSendMessageParams& ShowParams, const FOnShowSendMessageUIClosedDelegate& Delegate = FOnShowSendMessageUIClosedDelegate())
+	{
+		UE_LOG(LogOnlineExternalUI, Warning, TEXT("ShowSendMessageToUserUI is not implemented on this platform!"));
+		Delegate.ExecuteIfBound(false);
+		return false;
+	}
+
 	/**
 	 * Displays an informational system dialog.
 	 *
diff --git a/Engine/Plugins/Online/OnlineSubsystemSteam/Source/Private/OnlineExternalUIInterfaceSteam.cpp b/Engine/Plugins/Online/OnlineSubsystemSteam/Source/Private/OnlineExternalUIInterfaceSteam.cpp
--- a/Engine/Plugins/Online/OnlineSubsystemSteam/Source/Private/OnlineExternalUIInterfaceSteam.cpp
+++ b/Engine/Plugins/Online/OnlineSubsystemSteam/Source/Private/OnlineExternalUIInterfaceSteam.cpp
@@ -30,7 +30,18 @@ void FOnlineAsyncEventSteamExternalUITriggered::TriggerDelegates()
 		ExternalUISteam->ProfileUIClosedDelegate.ExecuteIfBound();
 		ExternalUISteam->ProfileUIClosedDelegate.Unbind();
 
-		//@todo samz - obtain final url
+		// We don't have a way to tell if you sent a message, but we attempt to send it for you.
+		ExternalUISteam->ShowMessageClosedDelegate.ExecuteIfBound(ExternalUISteam->bMessageSent);
+		ExternalUISteam->ShowMessageClosedDelegate.Unbind();
+		ExternalUISteam->bMessageSent = false;
+
+		// We don't have any way to know that you bought an item on the store from this overlay. 
+		// This would be handled either by a DLC query or the server WebAPI.
+		// This returns true in order to trigger license checks.
+		ExternalUISteam->ShowStoreClosedDelegate.ExecuteIfBound(true);
+		ExternalUISteam->ShowStoreClosedDelegate.Unbind();
+
+		// Steam doesn't allow you to capture the final browsing url on web overlays, so pass an empty string
 		ExternalUISteam->ShowWebUrlClosedDelegate.ExecuteIfBound(TEXT(""));
 		ExternalUISteam->ShowWebUrlClosedDelegate.Unbind();
 	}
@@ -41,6 +52,11 @@ bool FOnlineExternalUISteam::ShowLoginUI(const int ControllerIndex, bool bShowOn
 	return false;
 }
 
+bool FOnlineExternalUISteam::ShowAccountCreationUI(const int ControllerIndex, const FOnAccountCreationUIClosedDelegate& Delegate)
+{
+	return false;
+}
+
 bool FOnlineExternalUISteam::ShowFriendsUI(int32 LocalUserNum)
 {
 	SteamFriends()->ActivateGameOverlay("Friends");
@@ -122,11 +138,43 @@ bool FOnlineExternalUISteam::ShowAccountUpgradeUI(const FUniqueNetId& UniqueId)
 
 bool FOnlineExternalUISteam::ShowStoreUI(int32 LocalUserNum, const FShowStoreParams& ShowParams, const FOnShowStoreUIClosedDelegate& Delegate)
 {
-	return false;
+	if (!ShowParams.ProductId.IsNumeric() || ShowParams.ProductId.IsEmpty())
+	{
+		return false;
+	}
+
+	uint32 ProductId = (uint32)FCString::Atoi(*ShowParams.ProductId);
+
+	if (ProductId == 0)
+	{
+		return false;
+	}
+
+	SteamFriends()->ActivateGameOverlayToStore(ProductId, ShowParams.bAddToCart ? k_EOverlayToStoreFlag_AddToCartAndShow : k_EOverlayToStoreFlag_None);
+	ShowStoreClosedDelegate = Delegate;
+
+	return true;
 }
 
 bool FOnlineExternalUISteam::ShowSendMessageUI(int32 LocalUserNum, const FShowSendMessageParams& ShowParams, const FOnShowSendMessageUIClosedDelegate& Delegate)
 {
+	// Steam only allows an application to open the chat UI if a recipient is specified.
 	return false;
 }
 
+bool FOnlineExternalUISteam::ShowSendMessageToUserUI(int32 LocalUserNum, const FUniqueNetId& Recipient, const FShowSendMessageParams& ShowParams, const FOnShowSendMessageUIClosedDelegate& Delegate)
+{
+	const FUniqueNetIdSteam& TargetUser = (const FUniqueNetIdSteam&)Recipient;
+	const FString MessageToSend = ShowParams.DisplayMessage.ToString();
+
+	if (!TargetUser.IsValid() || MessageToSend.IsEmpty())
+	{
+		return false;
+	}
+	ShowMessageClosedDelegate = Delegate;
+
+	bMessageSent = SteamFriends()->ReplyToFriendMessage(TargetUser, TCHAR_TO_UTF8(*MessageToSend));
+	SteamFriends()->ActivateGameOverlayToUser(TCHAR_TO_UTF8(TEXT("chat")), TargetUser);
+	
+	return true;
+}
diff --git a/Engine/Plugins/Online/OnlineSubsystemSteam/Source/Private/OnlineExternalUIInterfaceSteam.h b/Engine/Plugins/Online/OnlineSubsystemSteam/Source/Private/OnlineExternalUIInterfaceSteam.h
--- a/Engine/Plugins/Online/OnlineSubsystemSteam/Source/Private/OnlineExternalUIInterfaceSteam.h
+++ b/Engine/Plugins/Online/OnlineSubsystemSteam/Source/Private/OnlineExternalUIInterfaceSteam.h
@@ -57,21 +57,24 @@ private:
 	class FOnlineSubsystemSteam* SteamSubsystem;
 
 	/** Hidden on purpose */
-	FOnlineExternalUISteam() :
-		SteamSubsystem(NULL)
-	{
-	}
+	FOnlineExternalUISteam() = delete;
 
 PACKAGE_SCOPE:
 
 	FOnlineExternalUISteam(class FOnlineSubsystemSteam* InSteamSubsystem) :
-		SteamSubsystem(InSteamSubsystem)
+		SteamSubsystem(InSteamSubsystem),
+		bMessageSent(false)
 	{
 	}
 
+	/** If a message was successfully sent via the Steam UI */
+	bool bMessageSent;
+
 	/** Triggered when the steam overlay is closed if it was opened via ShowProfileUI. Delegate will be unbound after it is executed. */
 	FOnProfileUIClosedDelegate ProfileUIClosedDelegate;
 	FOnShowWebUrlClosedDelegate ShowWebUrlClosedDelegate;
+	FOnShowSendMessageUIClosedDelegate ShowMessageClosedDelegate;
+	FOnShowStoreUIClosedDelegate ShowStoreClosedDelegate;
 
 public:
 
@@ -81,7 +84,7 @@ public:
 
 	// IOnlineExternalUI
 	virtual bool ShowLoginUI(const int ControllerIndex, bool bShowOnlineOnly, bool bShowSkipButton, const FOnLoginUIClosedDelegate& Delegate = FOnLoginUIClosedDelegate()) override;
-	virtual bool ShowAccountCreationUI(const int ControllerIndex, const FOnAccountCreationUIClosedDelegate& Delegate = FOnAccountCreationUIClosedDelegate()) override { /** NYI */ return false; }
+	virtual bool ShowAccountCreationUI(const int ControllerIndex, const FOnAccountCreationUIClosedDelegate& Delegate = FOnAccountCreationUIClosedDelegate()) override;
 	virtual bool ShowFriendsUI(int32 LocalUserNum) override;
 	virtual bool ShowInviteUI(int32 LocalUserNum, FName SessionName = NAME_GameSession) override;
 	virtual bool ShowAchievementsUI(int32 LocalUserNum) override;
@@ -92,6 +95,7 @@ public:
 	virtual bool ShowAccountUpgradeUI(const FUniqueNetId& UniqueId) override;
 	virtual bool ShowStoreUI(int32 LocalUserNum, const FShowStoreParams& ShowParams, const FOnShowStoreUIClosedDelegate& Delegate = FOnShowStoreUIClosedDelegate()) override;
 	virtual bool ShowSendMessageUI(int32 LocalUserNum, const FShowSendMessageParams& ShowParams, const FOnShowSendMessageUIClosedDelegate& Delegate = FOnShowSendMessageUIClosedDelegate()) override;
+	virtual bool ShowSendMessageToUserUI(int32 LocalUserNum, const FUniqueNetId& Recipient, const FShowSendMessageParams& ShowParams, const FOnShowSendMessageUIClosedDelegate& Delegate = FOnShowSendMessageUIClosedDelegate()) override;
 };
 
 typedef TSharedPtr<FOnlineExternalUISteam, ESPMode::ThreadSafe> FOnlineExternalUISteamPtr;
diff --git a/Engine/Plugins/Online/OnlineSubsystemUtils/Source/OnlineSubsystemUtils/Private/OnlineSubsystemUtils.cpp b/Engine/Plugins/Online/OnlineSubsystemUtils/Source/OnlineSubsystemUtils/Private/OnlineSubsystemUtils.cpp
--- a/Engine/Plugins/Online/OnlineSubsystemUtils/Source/OnlineSubsystemUtils/Private/OnlineSubsystemUtils.cpp
+++ b/Engine/Plugins/Online/OnlineSubsystemUtils/Source/OnlineSubsystemUtils/Private/OnlineSubsystemUtils.cpp
@@ -704,6 +704,8 @@ static bool OnlineExec( UWorld* InWorld, const TCHAR* Cmd, FOutputDevice& Ar )
 					{
 						// Full command usage:    EXTERNALUI ACHIEVEMENTS FRIENDS INVITE LOGIN PROFILE WEBURL
 						// Example for one test:  EXTERNALUI WEBURL
+						// Example for store: EXTERNALUI STORE productid true
+						// Example for send message: EXTERNALUI MESSAGE user "message"
 						// Note that tests are enabled in alphabetical order
 						bool bTestAchievementsUI = FParse::Command(&Cmd, TEXT("ACHIEVEMENTS")) ? true : false;
 						bool bTestFriendsUI = FParse::Command(&Cmd, TEXT("FRIENDS")) ? true : false;
@@ -713,7 +715,26 @@ static bool OnlineExec( UWorld* InWorld, const TCHAR* Cmd, FOutputDevice& Ar )
 						bool bTestWebURL = FParse::Command(&Cmd, TEXT("WEBURL")) ? true : false;
 
 						// This class also deletes itself once done
-						(new FTestExternalUIInterface(SubName, bTestLoginUI, bTestFriendsUI, bTestInviteUI, bTestAchievementsUI, bTestWebURL, bTestProfileUI))->Test();
+						FTestExternalUIInterface* TestHarness = (new FTestExternalUIInterface(SubName, bTestLoginUI, bTestFriendsUI, 
+							bTestInviteUI, bTestAchievementsUI, bTestWebURL, bTestProfileUI));
+
+						if (FParse::Command(&Cmd, TEXT("STORE")))
+						{
+							FString AppID = FParse::Token(Cmd, false);
+							bool bAddCart = FCString::ToBool(*FParse::Token(Cmd, false));
+							TestHarness->TestStorePage(AppID, bAddCart);
+						}
+						else if (FParse::Command(&Cmd, TEXT("MESSAGE")))
+						{
+							FString UserID = FParse::Token(Cmd, false);
+							FString Message = FParse::Token(Cmd, false);
+							TestHarness->TestSendMessage(UserID, Message);
+						}
+						else
+						{
+							TestHarness->Test();
+						}
+
 						bWasHandled = true;
 					}
 #endif //WITH_DEV_AUTOMATION_TESTS
diff --git a/Engine/Plugins/Online/OnlineSubsystemUtils/Source/OnlineSubsystemUtils/Private/Tests/TestExternalUIInterface.cpp b/Engine/Plugins/Online/OnlineSubsystemUtils/Source/OnlineSubsystemUtils/Private/Tests/TestExternalUIInterface.cpp
--- a/Engine/Plugins/Online/OnlineSubsystemUtils/Source/OnlineSubsystemUtils/Private/Tests/TestExternalUIInterface.cpp
+++ b/Engine/Plugins/Online/OnlineSubsystemUtils/Source/OnlineSubsystemUtils/Private/Tests/TestExternalUIInterface.cpp
@@ -2,22 +2,28 @@
 
 #include "Tests/TestExternalUIInterface.h"
 #include "OnlineSubsystem.h"
+#include "Internationalization/Text.h"
 #include "Interfaces/OnlineIdentityInterface.h"
 
 #if WITH_DEV_AUTOMATION_TESTS
 
-void FTestExternalUIInterface::Test()
+void FTestExternalUIInterface::PrepareTests()
 {
 	// Cache interfaces
 	OnlineSub = IOnlineSubsystem::Get(FName(*SubsystemName));
-	check(OnlineSub); 
+	check(OnlineSub);
 
 	ExternalUI = OnlineSub->GetExternalUIInterface();
 	check(ExternalUI.IsValid());
 
 	// Define and register delegates
-	ExternalUIChangeDelegate       = FOnExternalUIChangeDelegate::CreateRaw(this, &FTestExternalUIInterface::OnExternalUIChange);
+	ExternalUIChangeDelegate = FOnExternalUIChangeDelegate::CreateRaw(this, &FTestExternalUIInterface::OnExternalUIChange);
 	ExternalUIChangeDelegateHandle = ExternalUI->AddOnExternalUIChangeDelegate_Handle(ExternalUIChangeDelegate);
+}
+
+void FTestExternalUIInterface::Test()
+{
+	PrepareTests();
 
 	// Are we testing at least one of our external UI?
 	if (bTestLoginUI == false && bTestFriendsUI == false && bTestInviteUI == false && bTestAchievementsUI == false && bTestWebURL == false && bTestProfileUI == false)
@@ -31,6 +37,61 @@ void FTestExternalUIInterface::Test()
 	}
 }
 
+void FTestExternalUIInterface::TestSendMessage(const FString& InMsgRecepient, const FString& InMessage)
+{
+	PrepareTests();
+	bIsolatedTest = true;
+
+	if (InMsgRecepient.IsEmpty())
+	{
+		UE_LOG_ONLINE_EXTERNALUI(Error, TEXT("ExternalUI TestSendMessage -- InMsgRecepient must not be empty!"));
+		FinishTest();
+		return;
+	}
+
+	TSharedPtr<const FUniqueNetId> TargetUser = OnlineSub->GetIdentityInterface()->CreateUniquePlayerId(InMsgRecepient);
+	if (!TargetUser.IsValid())
+	{
+		UE_LOG_ONLINE_EXTERNALUI(Error, TEXT("ExternalUI TestSendMessage -- Message target user is invalid!"));
+		FinishTest();
+		return;
+	}
+
+	FShowSendMessageParams MessageParams;
+	MessageParams.DisplayMessage = FText::FromString(InMessage);
+
+	bool bWasOpened = ExternalUI->ShowSendMessageToUserUI(0, *TargetUser, MessageParams, FOnShowSendMessageUIClosedDelegate::CreateRaw(this, &FTestExternalUIInterface::OnSendMessageUIClosed));
+	UE_LOG_ONLINE_EXTERNALUI(Log, TEXT("ExternalUI TestSendMessage -- Message test opened %d"), bWasOpened);
+	if (!bWasOpened)
+	{
+		FinishTest();
+	}
+}
+
+void FTestExternalUIInterface::TestStorePage(const FString& InProductId, bool bShouldAddToCart)
+{
+	PrepareTests();
+	bIsolatedTest = true;
+
+	if (InProductId.IsEmpty())
+	{
+		UE_LOG_ONLINE_EXTERNALUI(Error, TEXT("ExternalUI TestStorePage -- Product Id must not be empty!"));
+		FinishTest();
+		return;
+	}
+
+	FShowStoreParams StoreParams;
+	StoreParams.ProductId = InProductId;
+	StoreParams.bAddToCart = bShouldAddToCart;
+
+	bool bWasOpened = ExternalUI->ShowStoreUI(0, StoreParams, FOnShowStoreUIClosedDelegate::CreateRaw(this, &FTestExternalUIInterface::OnStoreUIClosed));
+	UE_LOG_ONLINE_EXTERNALUI(Log, TEXT("ExternalUI TestStorePage -- Store test opened %d"), bWasOpened);
+	if (!bWasOpened)
+	{
+		FinishTest();
+	}
+}
+
 void FTestExternalUIInterface::FinishTest()
 {
 	UE_LOG_ONLINE_EXTERNALUI(Log, TEXT("FTestExternalUIInterface::FinishTest -- completed testing"));
@@ -177,9 +238,16 @@ void FTestExternalUIInterface::OnExternalUIChange(bool bIsOpening)
 
 	if (bIsOpening == false)
 	{
-		// The external UI is no longer active
-		// Move on to the next test
-		StartNextTest();
+		if (bIsolatedTest)
+		{
+			FinishTest();
+		}
+		else
+		{
+			// The external UI is no longer active
+			// Move on to the next test
+			StartNextTest();
+		}
 	}
 }
 
@@ -198,4 +266,14 @@ void FTestExternalUIInterface::OnShowWebUrlClosed(const FString& FinalUrl)
 	UE_LOG_ONLINE_EXTERNALUI(Log, TEXT("Show Web Url closed with FinalUrl=%s."), *FinalUrl);
 }
 
+void FTestExternalUIInterface::OnStoreUIClosed(bool bWasPurchased)
+{
+	UE_LOG_ONLINE_EXTERNALUI(Log, TEXT("Store UI was closed. Purchase Made: %d"), bWasPurchased);
+}
+
+void FTestExternalUIInterface::OnSendMessageUIClosed(bool bWasSent)
+{
+	UE_LOG_ONLINE_EXTERNALUI(Log, TEXT("Send Message UI closed. Message sent: %d"), bWasSent);
+}
+
 #endif //WITH_DEV_AUTOMATION_TESTS
diff --git a/Engine/Plugins/Online/OnlineSubsystemUtils/Source/OnlineSubsystemUtils/Private/Tests/TestExternalUIInterface.h b/Engine/Plugins/Online/OnlineSubsystemUtils/Source/OnlineSubsystemUtils/Private/Tests/TestExternalUIInterface.h
--- a/Engine/Plugins/Online/OnlineSubsystemUtils/Source/OnlineSubsystemUtils/Private/Tests/TestExternalUIInterface.h
+++ b/Engine/Plugins/Online/OnlineSubsystemUtils/Source/OnlineSubsystemUtils/Private/Tests/TestExternalUIInterface.h
@@ -58,6 +58,12 @@ namespace ETestExternalUIInterfaceState
 	/** Current external UI test */
 	ETestExternalUIInterfaceState::Type State;
 
+	/** If this is an one off test (for when calling a test that requires extra information, outside the automation loop) */
+	bool bIsolatedTest;
+
+	/** Prepares the testing harness environment */
+	void PrepareTests();
+
 	/** Completes testing and cleans up after itself */
 	void FinishTest();
 
@@ -84,6 +90,12 @@ namespace ETestExternalUIInterfaceState
 	/** Delegate executed when the show web UI has been closed. */
 	void OnShowWebUrlClosed(const FString& FinalUrl);
 
+	/** Delegate executed when the store UI has been closed. */
+	void OnStoreUIClosed(bool bWasPurchased);
+
+	/** Delegate executed when the send message UI has been closed. */
+	void OnSendMessageUIClosed(bool bWasSent);
+
  public:
 
 	/**
@@ -101,6 +113,7 @@ namespace ETestExternalUIInterfaceState
 		,	OnlineSub(NULL)
 		,	ExternalUI(NULL)
 		,	State(ETestExternalUIInterfaceState::Begin)
+		,	bIsolatedTest(false)
 	{
 	}
 
@@ -110,6 +123,16 @@ namespace ETestExternalUIInterfaceState
 	 */
 	void Test();
 
+	/**
+	 * Runs a test for sending messages (not automated as this requires user input)
+	 */
+	void TestSendMessage(const FString& InMsgRecepient, const FString& InMessage);
+
+	/**
+	 * Runs a test for opening the store page
+	 */
+	void TestStorePage(const FString& InProductId, bool bShouldAddToCart);
+
  };
 
 #endif //WITH_DEV_AUTOMATION_TESTS
