Blueprints: Fixed the error message in various online latent actions to reflect the correct action

diff --git a/Engine/Plugins/Online/OnlineSubsystemUtils/Source/OnlineSubsystemUtils/Private/EndMatchCallbackProxy.cpp b/Engine/Plugins/Online/OnlineSubsystemUtils/Source/OnlineSubsystemUtils/Private/EndMatchCallbackProxy.cpp
--- a/Engine/Plugins/Online/OnlineSubsystemUtils/Source/OnlineSubsystemUtils/Private/EndMatchCallbackProxy.cpp
+++ b/Engine/Plugins/Online/OnlineSubsystemUtils/Source/OnlineSubsystemUtils/Private/EndMatchCallbackProxy.cpp
@@ -35,7 +35,7 @@ UEndMatchCallbackProxy* UEndMatchCallbackProxy::EndMatch(UObject* WorldContextOb
 
 void UEndMatchCallbackProxy::Activate()
 {
-	FOnlineSubsystemBPCallHelper Helper(TEXT("ConnectToService"), WorldContextObject);
+	FOnlineSubsystemBPCallHelper Helper(TEXT("EndMatch"), WorldContextObject);
 	Helper.QueryIDFromPlayerController(PlayerControllerWeakPtr.Get());
 
 	if (Helper.IsValid())
diff --git a/Engine/Plugins/Online/OnlineSubsystemUtils/Source/OnlineSubsystemUtils/Private/EndTurnCallbackProxy.cpp b/Engine/Plugins/Online/OnlineSubsystemUtils/Source/OnlineSubsystemUtils/Private/EndTurnCallbackProxy.cpp
--- a/Engine/Plugins/Online/OnlineSubsystemUtils/Source/OnlineSubsystemUtils/Private/EndTurnCallbackProxy.cpp
+++ b/Engine/Plugins/Online/OnlineSubsystemUtils/Source/OnlineSubsystemUtils/Private/EndTurnCallbackProxy.cpp
@@ -34,7 +34,7 @@ UEndTurnCallbackProxy* UEndTurnCallbackProxy::EndTurn(UObject* WorldContextObjec
 
 void UEndTurnCallbackProxy::Activate()
 {
-	FOnlineSubsystemBPCallHelper Helper(TEXT("ConnectToService"), WorldContextObject);
+	FOnlineSubsystemBPCallHelper Helper(TEXT("EndTurn"), WorldContextObject);
 	Helper.QueryIDFromPlayerController(PlayerControllerWeakPtr.Get());
 
 	if (Helper.IsValid())
diff --git a/Engine/Plugins/Online/OnlineSubsystemUtils/Source/OnlineSubsystemUtils/Private/FindTurnBasedMatchCallbackProxy.cpp b/Engine/Plugins/Online/OnlineSubsystemUtils/Source/OnlineSubsystemUtils/Private/FindTurnBasedMatchCallbackProxy.cpp
--- a/Engine/Plugins/Online/OnlineSubsystemUtils/Source/OnlineSubsystemUtils/Private/FindTurnBasedMatchCallbackProxy.cpp
+++ b/Engine/Plugins/Online/OnlineSubsystemUtils/Source/OnlineSubsystemUtils/Private/FindTurnBasedMatchCallbackProxy.cpp
@@ -39,7 +39,7 @@ UFindTurnBasedMatchCallbackProxy* UFindTurnBasedMatchCallbackProxy::FindTurnBase
 
 void UFindTurnBasedMatchCallbackProxy::Activate()
 {
-	FOnlineSubsystemBPCallHelper Helper(TEXT("ConnectToService"), WorldContextObject);
+	FOnlineSubsystemBPCallHelper Helper(TEXT("FindTurnBasedMatch"), WorldContextObject);
 	Helper.QueryIDFromPlayerController(PlayerControllerWeakPtr.Get());
 
 	if (Helper.IsValid())
diff --git a/Engine/Plugins/Online/OnlineSubsystemUtils/Source/OnlineSubsystemUtils/Private/QuitMatchCallbackProxy.cpp b/Engine/Plugins/Online/OnlineSubsystemUtils/Source/OnlineSubsystemUtils/Private/QuitMatchCallbackProxy.cpp
--- a/Engine/Plugins/Online/OnlineSubsystemUtils/Source/OnlineSubsystemUtils/Private/QuitMatchCallbackProxy.cpp
+++ b/Engine/Plugins/Online/OnlineSubsystemUtils/Source/OnlineSubsystemUtils/Private/QuitMatchCallbackProxy.cpp
@@ -29,7 +29,7 @@ UQuitMatchCallbackProxy* UQuitMatchCallbackProxy::QuitMatch(UObject* WorldContex
 
 void UQuitMatchCallbackProxy::Activate()
 {
-	FOnlineSubsystemBPCallHelper Helper(TEXT("ConnectToService"), WorldContextObject);
+	FOnlineSubsystemBPCallHelper Helper(TEXT("QuitMatch"), WorldContextObject);
 	Helper.QueryIDFromPlayerController(PlayerControllerWeakPtr.Get());
 
 	if (Helper.IsValid())
