FRIENDS test refactor to improve error handling.

diff --git a/Engine/Plugins/Online/OnlineSubsystemUtils/Source/OnlineSubsystemUtils/Private/Tests/TestFriendsInterface.cpp b/Engine/Plugins/Online/OnlineSubsystemUtils/Source/OnlineSubsystemUtils/Private/Tests/TestFriendsInterface.cpp
--- a/Engine/Plugins/Online/OnlineSubsystemUtils/Source/OnlineSubsystemUtils/Private/Tests/TestFriendsInterface.cpp
+++ b/Engine/Plugins/Online/OnlineSubsystemUtils/Source/OnlineSubsystemUtils/Private/Tests/TestFriendsInterface.cpp
@@ -41,7 +41,8 @@ void FTestFriendsInterface::Test(UWorld* InWorld, const TArray<FString>& Invites
 			}
 		}
 
-		// kick off next test
+		UE_LOG_ONLINE_FRIEND(Log, TEXT("Starting friends interface tests for %s"), *SubsystemName);
+
 		StartNextTest();
 	}
 	else
@@ -67,54 +68,107 @@ void FTestFriendsInterface::StartNextTest()
 {
 	if (bReadFriendsList)
 	{
+		UE_LOG_ONLINE_FRIEND(Log, TEXT("Reading friends list"));
+
 		IOnlineSharingPtr SharingInt = OnlineSub->GetSharingInterface();
 		if (SharingInt.IsValid())
 		{
+			UE_LOG_ONLINE_FRIEND(Log, TEXT("Sharing interface valid. Requesting permissions"));
+
 			FOnRequestNewReadPermissionsCompleteDelegate Delegate = FOnRequestNewReadPermissionsCompleteDelegate::CreateRaw(this, &FTestFriendsInterface::OnRequestNewPermissionsComplete);
 			OnRequestNewReadPermissionsDelegateHandle = SharingInt->AddOnRequestNewReadPermissionsCompleteDelegate_Handle(0, Delegate);
 			SharingInt->RequestNewReadPermissions(0, EOnlineSharingCategory::Friends);
 		}
 		else
 		{
+			UE_LOG_ONLINE_FRIEND(Log, TEXT("Sharing interface invalid. Skipping permission request"));
+
 			OnRequestNewPermissionsComplete(0, true);
 		}
 	}
 	else if (bQueryRecentPlayers)
 	{
-		if (OnlineSub->GetIdentityInterface().IsValid() &&
-			OnlineSub->GetIdentityInterface()->GetUniquePlayerId(0).IsValid())
+		UE_LOG_ONLINE_FRIEND(Log, TEXT("Querying recent players"));
+
+		//Anvil doesn't have methods to query recent players in the gdk
+		if (SubsystemName.Equals("anvil", ESearchCase::IgnoreCase))
 		{
-			OnlineSub->GetFriendsInterface()->QueryRecentPlayers(*OnlineSub->GetIdentityInterface()->GetUniquePlayerId(0), RecentPlayersNamespace);
+			UE_LOG_ONLINE_FRIEND(Log, TEXT("The Anvil subsystem doesn't have methods to query recent players. Skipping test"));
+
+			bQueryRecentPlayers = false;
+			StartNextTest();
+		}
+		else
+		{
+			if (OnlineSub->GetIdentityInterface().IsValid() &&
+				OnlineSub->GetIdentityInterface()->GetUniquePlayerId(0).IsValid())
+			{
+				OnlineSub->GetFriendsInterface()->QueryRecentPlayers(*OnlineSub->GetIdentityInterface()->GetUniquePlayerId(0), RecentPlayersNamespace);
+			}
+			bQueryRecentPlayers = false;
 		}
-		bQueryRecentPlayers = false;
 	}
 	else if (bAcceptInvites && InvitesToAccept.Num() > 0)
 	{
+		UE_LOG_ONLINE_FRIEND(Log, TEXT("Accepting invites"));
+
 		FOnAcceptInviteComplete Delegate = FOnAcceptInviteComplete::CreateRaw(this, &FTestFriendsInterface::OnAcceptInviteComplete);
 		OnlineSub->GetFriendsInterface()->AcceptInvite(0, *InvitesToAccept[0], FriendsListName, Delegate);
 	}
 	else if (bSendInvites && InvitesToSend.Num() > 0)
 	{
+		UE_LOG_ONLINE_FRIEND(Log, TEXT("Sending invites"));
+
 		FOnSendInviteComplete OnSendInviteCompleteDelegate = FOnSendInviteComplete::CreateRaw(this, &FTestFriendsInterface::OnSendInviteComplete);
 		OnlineSub->GetFriendsInterface()->SendInvite(0, *InvitesToSend[0], FriendsListName, OnSendInviteCompleteDelegate);
 	}
 	else if (bDeleteFriends && FriendsToDelete.Num() > 0)
 	{
-		OnlineSub->GetFriendsInterface()->DeleteFriend(0, *FriendsToDelete[0], FriendsListName);
+		UE_LOG_ONLINE_FRIEND(Log, TEXT("Deleting friends"));
+
+		//Anvil doesn't have methods to remove friends in the gdk
+		if (SubsystemName.Equals("anvil", ESearchCase::IgnoreCase))
+		{
+			UE_LOG_ONLINE_FRIEND(Log, TEXT("The Anvil subsystem doesn't have methods to delete friends. Skipping test"));
+
+			bDeleteFriends = false;
+			StartNextTest();
+		}
+		else
+		{
+			OnlineSub->GetFriendsInterface()->DeleteFriend(0, *FriendsToDelete[0], FriendsListName);
+		}
 	}
 	else if (bDeleteFriendsList)
 	{
-		FOnDeleteFriendsListComplete Delegate = FOnDeleteFriendsListComplete::CreateRaw(this, &FTestFriendsInterface::OnDeleteFriendsListComplete);
-		OnlineSub->GetFriendsInterface()->DeleteFriendsList(0, FriendsListName, Delegate);
+		UE_LOG_ONLINE_FRIEND(Log, TEXT("Deleting friends list"));
+
+		//Anvil doesn't have methods to remove friend lists in the gdk
+		if (SubsystemName.Equals("anvil", ESearchCase::IgnoreCase))
+		{
+			UE_LOG_ONLINE_FRIEND(Log, TEXT("The Anvil subsystem doesn't have methods to delete friends lists. Skipping test"));
+
+			bDeleteFriendsList = false;
+			StartNextTest();
+		}
+		else
+		{
+			FOnDeleteFriendsListComplete Delegate = FOnDeleteFriendsListComplete::CreateRaw(this, &FTestFriendsInterface::OnDeleteFriendsListComplete);
+			OnlineSub->GetFriendsInterface()->DeleteFriendsList(0, FriendsListName, Delegate);
+		}
 	}
 	else if (bQueryBlockedPlayers)
 	{
+		UE_LOG_ONLINE_FRIEND(Log, TEXT("Querying blocked players"));
+
 		if (OnlineSub->GetIdentityInterface()->GetUniquePlayerId(0).IsValid())
 		{
 			OnlineSub->GetFriendsInterface()->QueryBlockedPlayers(*OnlineSub->GetIdentityInterface()->GetUniquePlayerId(0));
 		}
 		else
 		{
+			UE_LOG_ONLINE_FRIEND(Log, TEXT("Player unique id invalid. Skipping Test"));
+
 			bQueryBlockedPlayers = false;
 			StartNextTest();
 		}
@@ -127,6 +181,8 @@ void FTestFriendsInterface::StartNextTest()
 
 void FTestFriendsInterface::FinishTest()
 {
+	UE_LOG_ONLINE_FRIEND(Log, TEXT("All tests finished"));
+
 	if (OnlineSub != NULL &&
 		OnlineSub->GetFriendsInterface().IsValid())
 	{
@@ -140,8 +196,7 @@ void FTestFriendsInterface::FinishTest()
 
 void FTestFriendsInterface::OnReadFriendsComplete(int32 LocalPlayer, bool bWasSuccessful, const FString& ListName, const FString& ErrorStr)
 {
-	UE_LOG_ONLINE_FRIEND(Log,
-		TEXT("ReadFriendsList() for player (%d) was success=%d error=%s"), LocalPlayer, bWasSuccessful, *ErrorStr);
+	UE_LOG_ONLINE_FRIEND(Log, TEXT("ReadFriendsList() for player (%d) was success=%d error=%s"), LocalPlayer, bWasSuccessful, *ErrorStr);
 
 	if (bWasSuccessful)
 	{
@@ -306,7 +361,7 @@ void FTestFriendsInterface::OnDeleteFriendsListComplete(int32 LocalPlayer, bool
 
 void FTestFriendsInterface::OnQueryBlockedPlayersComplete(const FUniqueNetId& UserId, bool bWasSuccessful, const FString& Error)
 {
-	UE_LOG_ONLINE_FRIEND(Log, TEXT("QueryBlockedPlayers() for player (%s) was success=%d"), *UserId.ToDebugString(), bWasSuccessful);
+	UE_LOG_ONLINE_FRIEND(Log, TEXT("QueryBlockedPlayers() for player (%s) was success=%d error=%s"), *UserId.ToDebugString(), bWasSuccessful, *Error);
 	bQueryBlockedPlayers = false;
 
 	if (bWasSuccessful)
