PRESENCE test refactor to improve error handling.

diff --git a/Engine/Plugins/Online/OnlineSubsystemUtils/Source/OnlineSubsystemUtils/Private/Tests/TestPresenceInterface.cpp b/Engine/Plugins/Online/OnlineSubsystemUtils/Source/OnlineSubsystemUtils/Private/Tests/TestPresenceInterface.cpp
--- a/Engine/Plugins/Online/OnlineSubsystemUtils/Source/OnlineSubsystemUtils/Private/Tests/TestPresenceInterface.cpp
+++ b/Engine/Plugins/Online/OnlineSubsystemUtils/Source/OnlineSubsystemUtils/Private/Tests/TestPresenceInterface.cpp
@@ -161,6 +161,8 @@ bool FTestPresenceInterface::Tick(float DeltaTime)
 
 void FTestPresenceInterface::Test(UWorld* InWorld, const FString& RandomUser)
 {
+	UE_LOG_ONLINE_PRESENCE(Log, TEXT("Starting Presence Interface Tests"));
+
 	// Grab the OSS
 	OnlineSub = Online::GetSubsystem(InWorld, FName(*SubsystemName));
 	check(OnlineSub);
@@ -213,6 +215,8 @@ void FTestPresenceInterface::Test(UWorld* InWorld, const FString& RandomUser)
 
 	// Grab the user's friends
 	// This will then random from the friends list and fetch a profile from here.
+	UE_LOG_ONLINE_PRESENCE(Log, TEXT("Reading Friends List"));
+
 	FriendsInt->ReadFriendsList(0, EFriendsLists::ToString(EFriendsLists::Default), FOnReadFriendsListComplete::CreateRaw(this, &FTestPresenceInterface::OnReadFriendsListComplete));
 
 	// First, fetch your own profile.
@@ -307,7 +311,6 @@ bool FTestPresenceInterface::DoesPresenceMatchWithCache(const TSharedPtr<FOnline
 		return false;
 	}
 
-	FString TestValue;
 	for (FPresenceProperties::TConstIterator Itr(CachedPresence.Properties); Itr; ++Itr)
 	{
 		const FString& Key = Itr.Key();
@@ -318,7 +321,7 @@ bool FTestPresenceInterface::DoesPresenceMatchWithCache(const TSharedPtr<FOnline
 			return false;
 		}
 
-		PresenceToTest->Status.Properties.Find(Key)->GetValue(TestValue);
+		const FString& TestValue = PresenceToTest->Status.Properties.Find(Key)->ToString();
 		if (Value != TestValue)
 		{
 			UE_LOG_ONLINE_PRESENCE(Warning, TEXT("Presence test fails, key %s has different values. Cached=%s Has=%s"), *Key, *Value, *TestValue);
@@ -331,7 +334,7 @@ bool FTestPresenceInterface::DoesPresenceMatchWithCache(const TSharedPtr<FOnline
 
 void FTestPresenceInterface::OnSelfPresenceFetchComplete(const FUniqueNetId& UserId, const bool bWasSuccessful)
 {
-	UE_LOG_ONLINE_PRESENCE(Log, TEXT("FTestPresenceInterface::OnSelfPresenceFetchComplete Succeeded? %d"), bWasSuccessful);
+	UE_LOG_ONLINE_PRESENCE(Log, TEXT("FTestPresenceInterface::OnSelfPresenceFetchComplete Success=%d"), bWasSuccessful);
 	TasksAttempted |= EPresenceTestStatus::FetchSelf;
 	if (bWasSuccessful)
 	{
@@ -370,7 +373,7 @@ void FTestPresenceInterface::OnPresencePushComplete(const FUniqueNetId& UserId,
 	}
 	
 	TasksAttempted |= EPresenceTestStatus::PushSelf;
-	UE_LOG_ONLINE_PRESENCE(Log, TEXT("FTestPresenceInterface::OnPresencePushComplete Task Succeeded? %d Data correct? %d"), bWasSuccessful, bPushCorrect);
+	UE_LOG_ONLINE_PRESENCE(Log, TEXT("FTestPresenceInterface::OnPresencePushComplete Task Success=%d. Data correct=%d"), bWasSuccessful, bPushCorrect);
 }
 
 void FTestPresenceInterface::OnPresencePushWithSessionComplete(const FUniqueNetId& UserId, const bool bWasSuccessful)
@@ -397,12 +400,13 @@ void FTestPresenceInterface::OnPresencePushWithSessionComplete(const FUniqueNetI
 	}
 
 	TasksAttempted |= EPresenceTestStatus::PushSelfSession;
-	UE_LOG_ONLINE_PRESENCE(Log, TEXT("FTestPresenceInterface::OnPresencePushWithSessionComplete Task Succeeded? %d Data correct? %d"), bWasSuccessful, bPushCorrect);
+	UE_LOG_ONLINE_PRESENCE(Log, TEXT("FTestPresenceInterface::OnPresencePushWithSessionComplete Task Success=%d. Data correct=%d"), bWasSuccessful, bPushCorrect);
 }
 
 void FTestPresenceInterface::OnFriendPresenceFetchComplete(const FUniqueNetId& UserId, const bool bWasSuccessful)
 {
-	UE_LOG_ONLINE_PRESENCE(Log, TEXT("FTestPresenceInterface::OnFriendPresenceFetchComplete. User: %s Succeeded? %d"), *UserId.ToString(), bWasSuccessful);
+	UE_LOG_ONLINE_PRESENCE(Log, TEXT("FTestPresenceInterface::OnFriendPresenceFetchComplete. User: %s. Success=%d"), *UserId.ToString(), bWasSuccessful);
+
 	TasksAttempted |= EPresenceTestStatus::FetchFriend;
 	if (bWasSuccessful)
 	{
@@ -419,7 +423,7 @@ void FTestPresenceInterface::OnFriendPresenceFetchComplete(const FUniqueNetId& U
 
 void FTestPresenceInterface::OnRandomUserFetchComplete(const FUniqueNetId& UserId, const bool bWasSuccessful)
 {
-	UE_LOG_ONLINE_PRESENCE(Log, TEXT("FTestPresenceInterface::OnRandomUserFetchComplete User: %s Succeeded? %d"), *UserId.ToString(), bWasSuccessful);
+	UE_LOG_ONLINE_PRESENCE(Log, TEXT("FTestPresenceInterface::OnRandomUserFetchComplete User: %s Success=%d"), *UserId.ToString(), bWasSuccessful);
 	TasksAttempted |= EPresenceTestStatus::FetchRandom;
 	if (bWasSuccessful)
 	{
@@ -444,20 +448,30 @@ void FTestPresenceInterface::OnPresenceRecieved(const FUniqueNetId& UserId, cons
 
 void FTestPresenceInterface::OnReadFriendsListComplete(int32 LocalUserNum, bool bWasSuccessful, const FString& ListName, const FString& ErrorStr)
 {
-	UE_LOG_ONLINE_PRESENCE(Log, TEXT("Friends list fetch task for presence test complete. Success? %d"), bWasSuccessful);
+	UE_LOG_ONLINE_PRESENCE(Log, TEXT("Friends list fetch task for presence test complete. Success=%d. Error=%s"), bWasSuccessful, *ErrorStr);
+
 	if (bWasSuccessful)
 	{
 		FriendsInt->GetFriendsList(0, EFriendsLists::ToString(EFriendsLists::Default), FriendsCache);
+
+		UE_LOG_ONLINE_PRESENCE(Log, TEXT("Found %d friends!"), FriendsCache.Num());
+
 		if (FriendsCache.Num() > 0)
 		{
 			int32 FriendToFetch = FMath::RandRange(0, FriendsCache.Num() - 1);
 			if (FriendsCache.IsValidIndex(FriendToFetch))
 			{
+				UE_LOG_ONLINE_PRESENCE(Log, TEXT("Fetching presence for friend number %d"), FriendToFetch);
+
 				// Start the friends data task
-				PresenceInt->QueryPresence(FriendsCache[FriendToFetch]->GetUserId().Get(),
-					IOnlinePresence::FOnPresenceTaskCompleteDelegate::CreateRaw(this, &FTestPresenceInterface::OnFriendPresenceFetchComplete));
+				PresenceInt->QueryPresence(FriendsCache[FriendToFetch]->GetUserId().Get(), IOnlinePresence::FOnPresenceTaskCompleteDelegate::CreateRaw(this, &FTestPresenceInterface::OnFriendPresenceFetchComplete));
+
 				return;
 			}
+			else
+			{
+				UE_LOG_ONLINE_PRESENCE(Log, TEXT("Friend index not valid: %d"), FriendToFetch);
+			}
 		}
 		else
 		{
