When removing players from a party reservation give still-connected players their own reservation to prevent them from being orphaned on the PartyBeaconHost.

#tests Compile DebugGame Editor Win64 / Development Editor Win64 / Shipping Client / Test Client / Ran -game and verified kick flow works as expected. Verified standard client leaving.

diff --git a/Engine/Plugins/Online/OnlineSubsystemUtils/Source/OnlineSubsystemUtils/Private/PartyBeaconHost.cpp b/Engine/Plugins/Online/OnlineSubsystemUtils/Source/OnlineSubsystemUtils/Private/PartyBeaconHost.cpp
--- a/Engine/Plugins/Online/OnlineSubsystemUtils/Source/OnlineSubsystemUtils/Private/PartyBeaconHost.cpp
+++ b/Engine/Plugins/Online/OnlineSubsystemUtils/Source/OnlineSubsystemUtils/Private/PartyBeaconHost.cpp
@@ -748,11 +748,11 @@ EPartyReservationResult::Type APartyBeaconHost::UpdatePartyReservation(const FPa
 		if (bIsRemovingMembers)
 		{
 			UE_LOG(LogPartyBeacon, Verbose, TEXT("Removing Member"));
-			const int32 ExistingReservationIdx = State->GetExistingReservation(ReservationUpdateRequest.PartyLeader);
+			int32 ExistingReservationIdx = State->GetExistingReservation(ReservationUpdateRequest.PartyLeader);
 			if (ExistingReservationIdx != INDEX_NONE)
 			{
 				TArray<FPartyReservation>& Reservations = State->GetReservations();
-				FPartyReservation& ExistingReservation = Reservations[ExistingReservationIdx];
+				FPartyReservation* ExistingReservation = &Reservations[ExistingReservationIdx];
 
 				// Read the list of players and remove the ones that are not in this party
 				TArray<FPlayerReservation> PlayersToDelete;
@@ -776,15 +776,44 @@ EPartyReservationResult::Type APartyBeaconHost::UpdatePartyReservation(const FPa
 				// Copy new player entries into existing reservation
 				for (int32 PlayerIdx = 0; PlayerIdx < PlayersToDelete.Num(); PlayerIdx++)
 				{
-					const FPlayerReservation& PlayerRes = PlayersToDelete[PlayerIdx];
-					ExistingReservation.RemoveAllPartyMembers(PlayerRes);
+					// If we're waiting for them to connect remove their reservation so they don't block other players from joining.
+					FUniqueNetIdMatcher PlayerMatch(*PlayersToDelete[PlayerIdx].UniqueId);
+					int32 FoundIdx = State->PlayersPendingJoin.IndexOfByPredicate(PlayerMatch);
+					if (FoundIdx != INDEX_NONE)
+					{
+						const FPlayerReservation& PlayerRes = PlayersToDelete[PlayerIdx];
+						ExistingReservation->RemoveAllPartyMembers(PlayerRes);
 
-					// Keep track of newly added players
-					PlayerRemoved(PlayerRes);
-					State->SanityCheckReservations(true);
+						PlayerRemoved(PlayerRes);
+						State->SanityCheckReservations(true);
+						State->NumConsumedReservations -= 1;
+					}
+					// Connected players should get their own new reservation.
+					else
+					{
+						UpdatePartyLeader(PlayersToDelete[PlayerIdx].UniqueId, PlayersToDelete[PlayerIdx].UniqueId);
+
+						// UpdatePartyLeader should modify the reservations array by adding a new reservation and potentially invalidating ExistingReservation's address. Sanity check the index as well.
+						const int32 NewReservationIdx = State->GetExistingReservation(ReservationUpdateRequest.PartyLeader);
+						if (NewReservationIdx != INDEX_NONE)
+						{
+							ExistingReservation = &Reservations[NewReservationIdx];
+
+							if (NewReservationIdx != ExistingReservationIdx)
+							{
+								UE_LOG(LogPartyBeacon, Verbose, TEXT("%s - Call to UpdatePartyLeader resulted in reservation index of party leader %s changing from %d to %d."), ANSI_TO_TCHAR(__FUNCTION__), *ReservationUpdateRequest.PartyLeader->ToString(), ExistingReservationIdx, NewReservationIdx);
+								ExistingReservationIdx = NewReservationIdx;
+							}
+						}
+						// Existing reservation was removed.
+						else
+						{
+							UE_LOG(LogPartyBeacon, Warning, TEXT("%s - Call to UpdatePartyLeader resulted in removal of reservation with %s as the party leader."), ANSI_TO_TCHAR(__FUNCTION__), *ReservationUpdateRequest.PartyLeader->ToString());
+							break;
+						}
+					}
 				}
-				// Update the reservation count before sending the response
-				State->NumConsumedReservations -= PlayersToDelete.Num();
+
 				UE_LOG(LogPartyBeacon, Verbose, TEXT("APartyBeaconHost::UpdatePartyReservation: Removed %d players, setting NumConsumedReservations to %d"), PlayersToDelete.Num(), State->NumConsumedReservations);
 
 				// Tell any UI and/or clients that there has been a change in the reservation state
