SHARING test refactor adding check for Sharing Interface implementation

diff --git a/Engine/Plugins/Online/OnlineSubsystemUtils/Source/OnlineSubsystemUtils/Private/Tests/TestSharingInterface.cpp b/Engine/Plugins/Online/OnlineSubsystemUtils/Source/OnlineSubsystemUtils/Private/Tests/TestSharingInterface.cpp
--- a/Engine/Plugins/Online/OnlineSubsystemUtils/Source/OnlineSubsystemUtils/Private/Tests/TestSharingInterface.cpp
+++ b/Engine/Plugins/Online/OnlineSubsystemUtils/Source/OnlineSubsystemUtils/Private/Tests/TestSharingInterface.cpp
@@ -34,17 +34,23 @@ void FTestSharingInterface::Test(UWorld* InWorld, bool bWithImage)
 	check(OnlineSub); 
 
 	SharingInterface = OnlineSub->GetSharingInterface();
-	check(SharingInterface.IsValid());
+	if (SharingInterface.IsValid())
+	{
+		TestStatusUpdate.Message = FString::Printf(TEXT("This is a test post for UE4 Sharing support! Date = %s"), *FDateTime::Now().ToString());
+		TestStatusUpdate.PostPrivacy = EOnlineStatusUpdatePrivacy::OnlyMe;
+		if (bWithImage)
+		{
+			TestStatusUpdate.Image = new FImage(256, 256, ERawImageFormat::BGRA8, EGammaSpace::Linear);
+		}
 
-	TestStatusUpdate.Message = FString::Printf(TEXT("This is a test post for UE4 Sharing support! Date = %s"), *FDateTime::Now().ToString());
-	TestStatusUpdate.PostPrivacy = EOnlineStatusUpdatePrivacy::OnlyMe;
-	if( bWithImage )
+		// Kick off the first part of the test,
+		RequestPermissionsToSharePosts();
+	}
+	else
 	{
-		TestStatusUpdate.Image = new FImage( 256, 256, ERawImageFormat::BGRA8, EGammaSpace::Linear );
+		UE_LOG_ONLINE_SHARING(Display, TEXT("FTestSharingInterface TESTS SKIPPED. Sharing Interface implementation not found"));
+		delete this;
 	}
-
-	// Kick off the first part of the test,
-	RequestPermissionsToSharePosts();
 }
 
 
