Make sure UIpConnection::Socket is set to null when the socket is destroyed if it was also in the BindSockets list.
Workaround for a crash that can occur when destroying a net driver.

diff --git a/Engine/Plugins/Online/OnlineSubsystemUtils/Source/OnlineSubsystemUtils/Private/IpConnection.cpp b/Engine/Plugins/Online/OnlineSubsystemUtils/Source/OnlineSubsystemUtils/Private/IpConnection.cpp
--- a/Engine/Plugins/Online/OnlineSubsystemUtils/Source/OnlineSubsystemUtils/Private/IpConnection.cpp
+++ b/Engine/Plugins/Online/OnlineSubsystemUtils/Source/OnlineSubsystemUtils/Private/IpConnection.cpp
@@ -606,6 +606,10 @@ void UIpConnection::CleanupResolutionSockets()
 			if (CurSocket != Socket || bCleanAll)
 			{
 				SocketSubsystem->DestroySocket(CurSocket);
+				if (CurSocket == Socket)
+				{
+					Socket = nullptr;
+				}
 				CurSocket = nullptr;
 			}
 		}
