#online

Added tests for Store and Purchase OSS interfaces

usage:
"online sub=<OSS name> test store <list of product ids>" - eg. "online sub=live test store 9MWBK8Z14MXD" to get product details for ShooterGame store offer 9MWBK8Z14MXD
"online sub=<OSS name> test purchase <store namespace> <list of product ids>" - eg. "online sub=live test purchase UNUSED 9MWBK8Z14MXD" to initiate checkout for ShooterGame store offer 9MWBK8Z14MXD

#tests running exec commands against ShooterGame on MS platform

diff --git a/Engine/Plugins/Online/OnlineSubsystemUtils/Source/OnlineSubsystemUtils/Private/OnlineSubsystemUtils.cpp b/Engine/Plugins/Online/OnlineSubsystemUtils/Source/OnlineSubsystemUtils/Private/OnlineSubsystemUtils.cpp
--- a/Engine/Plugins/Online/OnlineSubsystemUtils/Source/OnlineSubsystemUtils/Private/OnlineSubsystemUtils.cpp
+++ b/Engine/Plugins/Online/OnlineSubsystemUtils/Source/OnlineSubsystemUtils/Private/OnlineSubsystemUtils.cpp
@@ -36,6 +36,8 @@
 #include "Tests/TestVoice.h"
 #include "Tests/TestExternalUIInterface.h"
 #include "Tests/TestPresenceInterface.h"
+#include "Tests/TestStoreInterface.h"
+#include "Tests/TestPurchaseInterface.h"
 
 
 static FAutoConsoleCommand GSendRemoteTalkersToEndpointCommand(
@@ -729,6 +731,29 @@ static bool OnlineExec( UWorld* InWorld, const TCHAR* Cmd, FOutputDevice& Ar )
 
 						bWasHandled = true;
 					}
+					else if (FParse::Command(&Cmd, TEXT("STORE")))
+					{
+						TArray<FString> OfferIds;
+						for (FString OfferId = FParse::Token(Cmd, false); !OfferId.IsEmpty(); OfferId = FParse::Token(Cmd, false))
+						{
+							OfferIds.Add(OfferId);
+						}
+						// This class deletes itself once done
+						(new FTestStoreInterface(SubName))->Test(InWorld, OfferIds);
+						bWasHandled = true;
+					}
+					else if (FParse::Command(&Cmd, TEXT("PURCHASE")))
+					{
+						FString Namespace = FParse::Token(Cmd, false);
+						TArray<FString> OfferIds;
+						for (FString OfferId = FParse::Token(Cmd, false); !OfferId.IsEmpty(); OfferId = FParse::Token(Cmd, false))
+						{
+							OfferIds.Add(OfferId);
+						}
+						// This class deletes itself once done
+						(new FTestPurchaseInterface(SubName))->Test(InWorld, Namespace, OfferIds);
+						bWasHandled = true;
+					}
 #endif //WITH_DEV_AUTOMATION_TESTS
 				}
 				else if (FParse::Command(&Cmd, TEXT("VOICE")))
diff --git a/Engine/Plugins/Online/OnlineSubsystemUtils/Source/OnlineSubsystemUtils/Private/Tests/TestPurchaseInterface.cpp b/Engine/Plugins/Online/OnlineSubsystemUtils/Source/OnlineSubsystemUtils/Private/Tests/TestPurchaseInterface.cpp
new file mode 100644
--- /dev/null
+++ b/Engine/Plugins/Online/OnlineSubsystemUtils/Source/OnlineSubsystemUtils/Private/Tests/TestPurchaseInterface.cpp
@@ -0,0 +1,139 @@
+// Copyright Epic Games, Inc. All Rights Reserved.
+
+
+#include "Tests/TestPurchaseInterface.h"
+#include "OnlineSubsystemUtils.h"
+
+#if WITH_DEV_AUTOMATION_TESTS
+
+FTestPurchaseInterface::FTestPurchaseInterface(const FString& InSubsystem)
+	: OnlineSub(nullptr)
+{
+	UE_LOG_ONLINE_PURCHASE(Display, TEXT("FTestPurchaseInterface::FTestPurchaseInterface"));
+	SubsystemName = InSubsystem;
+}
+
+FTestPurchaseInterface::~FTestPurchaseInterface()
+{
+	UE_LOG_ONLINE_PURCHASE(Display, TEXT("FTestPurchaseInterface::~FTestPurchaseInterface"));
+}
+
+void FTestPurchaseInterface::Test(UWorld* InWorld, const FString& Namespace, const TArray<FString>& InOffersIds)
+{
+	UE_LOG_ONLINE_PURCHASE(Display, TEXT("FTestPurchaseInterface::Test"));
+
+	OnlineSub = Online::GetSubsystem(InWorld, SubsystemName.Len() ? FName(*SubsystemName, FNAME_Find) : NAME_None);
+	if (OnlineSub != nullptr)
+	{
+		if (OnlineSub->GetIdentityInterface().IsValid())
+		{
+			LocalUserId = OnlineSub->GetIdentityInterface()->GetUniquePlayerId(0);
+		}
+		if (LocalUserId.IsValid())
+		{
+			if (OnlineSub->GetPurchaseInterface().IsValid())
+			{
+				if (InOffersIds.Num() > 0)
+				{
+					// list of offers
+					for (const FString& OfferId : InOffersIds)
+					{
+						CheckoutRequest.AddPurchaseOffer(FOfferNamespace(Namespace), FUniqueOfferId(OfferId), 1);
+					}
+					// kick off next test
+					StartNextTest();
+					return;
+				}
+				else
+				{
+					UE_LOG_ONLINE_PURCHASE(Warning,
+						TEXT("No offer ids specified for purchase"));
+				}
+			}
+			else
+			{
+				UE_LOG_ONLINE_PURCHASE(Warning,
+					TEXT("Failed to get Purchase interface for %s"), *SubsystemName);
+			}
+		}
+		else
+		{
+			UE_LOG_ONLINE_PURCHASE(Warning,
+				TEXT("No logged in user."));
+		}
+	}
+	else
+	{
+		UE_LOG_ONLINE_PURCHASE(Warning,
+			TEXT("Failed to get OSS interface for %s."), *SubsystemName);
+	}
+	// finish due to failure
+	FinishTest();
+}
+
+void FTestPurchaseInterface::StartNextTest()
+{
+	if (!OnlineSub->GetPurchaseInterface()->IsAllowedToPurchase(*LocalUserId))
+	{
+		UE_LOG_ONLINE_PURCHASE(Warning,
+			TEXT("Purchasing not allowed for user %s"), *LocalUserId->ToDebugString());
+		
+		FinishTest();
+	}
+	else
+	{
+		// query receipts prior to checkout
+		OnlineSub->GetPurchaseInterface()->QueryReceipts(*LocalUserId, true,
+			FOnQueryReceiptsComplete::CreateRaw(this, &FTestPurchaseInterface::OnQueryReceiptsPreCheckout));
+	}
+}
+
+static void DumpReceiptInfo(const TArray<FPurchaseReceipt>& Receipts)
+{
+	UE_LOG_ONLINE_PURCHASE(Log, TEXT("Receipts=[%d]"), Receipts.Num());
+	for (const FPurchaseReceipt& Receipt : Receipts)
+	{
+		UE_LOG_ONLINE_PURCHASE(Log, TEXT("%s"), *Receipt.ToJson());
+	}
+}
+
+void FTestPurchaseInterface::OnQueryReceiptsPreCheckout(const FOnlineError& Result)
+{
+	UE_LOG_ONLINE_PURCHASE(Log, TEXT("Query receipts pre-checkout completed with Result=[%s]"), *Result.ToLogString());
+
+	TArray<FPurchaseReceipt> Receipts;
+	OnlineSub->GetPurchaseInterface()->GetReceipts(*LocalUserId, Receipts);
+	DumpReceiptInfo(Receipts);	
+
+	// purchase requested offers
+	OnlineSub->GetPurchaseInterface()->Checkout(*LocalUserId, CheckoutRequest,
+		FOnPurchaseCheckoutComplete::CreateRaw(this, &FTestPurchaseInterface::OnCheckoutComplete));
+}
+
+void FTestPurchaseInterface::OnCheckoutComplete(const FOnlineError& Result, const TSharedRef<FPurchaseReceipt>& Receipt)
+{
+	UE_LOG_ONLINE_PURCHASE(Log, TEXT("Checkout completed with Result=[%s]"), *Result.ToLogString());
+
+	// query receipts again after checkout
+	OnlineSub->GetPurchaseInterface()->QueryReceipts(*LocalUserId, true,
+		FOnQueryReceiptsComplete::CreateRaw(this, &FTestPurchaseInterface::OnQueryReceiptsPostCheckout));
+}
+
+void FTestPurchaseInterface::OnQueryReceiptsPostCheckout(const FOnlineError& Result)
+{
+	UE_LOG_ONLINE_PURCHASE(Log, TEXT("Query receipts post-checkout completed with Result=[%s]"), *Result.ToLogString());
+
+	TArray<FPurchaseReceipt> Receipts;
+	OnlineSub->GetPurchaseInterface()->GetReceipts(*LocalUserId, Receipts);
+	DumpReceiptInfo(Receipts);
+
+	// completed
+	FinishTest();
+}
+
+void FTestPurchaseInterface::FinishTest()
+{
+	delete this;
+}
+
+#endif //WITH_DEV_AUTOMATION_TESTS
diff --git a/Engine/Plugins/Online/OnlineSubsystemUtils/Source/OnlineSubsystemUtils/Private/Tests/TestPurchaseInterface.h b/Engine/Plugins/Online/OnlineSubsystemUtils/Source/OnlineSubsystemUtils/Private/Tests/TestPurchaseInterface.h
new file mode 100644
--- /dev/null
+++ b/Engine/Plugins/Online/OnlineSubsystemUtils/Source/OnlineSubsystemUtils/Private/Tests/TestPurchaseInterface.h
@@ -0,0 +1,74 @@
+// Copyright Epic Games, Inc. All Rights Reserved.
+
+#pragma once
+
+// UE4 includes
+
+#include "CoreMinimal.h"
+#include "UObject/CoreOnline.h"
+
+// Module includes
+#include "Interfaces/OnlinePurchaseInterface.h"
+
+class IOnlineSubsystem;
+
+#if WITH_DEV_AUTOMATION_TESTS
+
+/**
+ * Class used to test the Purchase interface
+ */
+class FTestPurchaseInterface
+{
+
+private:
+
+	/** The subsystem that was requested to be tested or the default if empty */
+	FString SubsystemName;
+	/** The online interface to use for testing */
+	IOnlineSubsystem* OnlineSub;
+	/** Logged in user */
+	TSharedPtr<const FUniqueNetId> LocalUserId;
+	/** Contains the product id to be purchased */
+	FPurchaseCheckoutRequest CheckoutRequest;
+	
+	/**
+	 * Step through the various tests that should be run and initiate the next one
+	 */
+	void StartNextTest();
+
+	/**
+	 * Finish/cleanup the tests
+	 */
+	void FinishTest();
+
+	/** checkout completion */
+	void OnCheckoutComplete(const FOnlineError& Result, const TSharedRef<FPurchaseReceipt>& Receipt);
+	/** receipt query completion */
+	void OnQueryReceiptsPreCheckout(const FOnlineError& Result);
+	/** receipt query completion */
+	void OnQueryReceiptsPostCheckout(const FOnlineError& Result);
+
+public:
+
+	/**
+	 *	Constructor which sets the subsystem name to test
+	 *
+	 * @param InSubsystem the subsystem to test
+	 */
+	FTestPurchaseInterface(const FString& InSubsystem);
+
+	/**
+	 *	Destructor
+	 */
+	~FTestPurchaseInterface();
+
+	/**
+	 *	Kicks off all of the testing process
+	 *
+	 * @param Namespace for the offers
+	 * @param InPurchaseIds list of offer/product ids to to purchase
+	 */
+	void Test(class UWorld* InWorld, const FString& Namespace, const TArray<FString>& InOffersIds);
+ };
+
+#endif //WITH_DEV_AUTOMATION_TESTS
diff --git a/Engine/Plugins/Online/OnlineSubsystemUtils/Source/OnlineSubsystemUtils/Private/Tests/TestStoreInterface.cpp b/Engine/Plugins/Online/OnlineSubsystemUtils/Source/OnlineSubsystemUtils/Private/Tests/TestStoreInterface.cpp
new file mode 100644
--- /dev/null
+++ b/Engine/Plugins/Online/OnlineSubsystemUtils/Source/OnlineSubsystemUtils/Private/Tests/TestStoreInterface.cpp
@@ -0,0 +1,117 @@
+// Copyright Epic Games, Inc. All Rights Reserved.
+
+
+#include "Tests/TestStoreInterface.h"
+#include "OnlineSubsystemUtils.h"
+
+#if WITH_DEV_AUTOMATION_TESTS
+
+FTestStoreInterface::FTestStoreInterface(const FString& InSubsystem)
+	: OnlineSub(nullptr)
+{
+	UE_LOG_ONLINE_STOREV2(Display, TEXT("FTestStoreInterface::FTestStoreInterface"));
+	SubsystemName = InSubsystem;
+}
+
+FTestStoreInterface::~FTestStoreInterface()
+{
+	UE_LOG_ONLINE_STOREV2(Display, TEXT("FTestStoreInterface::~FTestStoreInterface"));
+}
+
+void FTestStoreInterface::Test(UWorld* InWorld, const TArray<FString>& InOffersIds)
+{
+	UE_LOG_ONLINE_STOREV2(Display, TEXT("FTestStoreInterface::Test"));
+
+	OnlineSub = Online::GetSubsystem(InWorld, SubsystemName.Len() ? FName(*SubsystemName, FNAME_Find) : NAME_None);
+	if (OnlineSub != nullptr)
+	{
+		if (OnlineSub->GetIdentityInterface().IsValid())
+		{
+			LocalUserId = OnlineSub->GetIdentityInterface()->GetUniquePlayerId(0);
+		}
+		if (LocalUserId.IsValid())
+		{
+			if (OnlineSub->GetStoreV2Interface().IsValid())
+			{
+				// list of offers 
+				for (const FString& OfferId : InOffersIds)
+				{
+					RequestOfferIds.Add(FUniqueOfferId(OfferId));
+				}				
+				// kick off next test
+				StartNextTest();
+				return;
+			}
+			else
+			{
+				UE_LOG_ONLINE_STOREV2(Warning,
+					TEXT("Failed to get Store interface for %s"), *SubsystemName);
+			}
+		}
+		else
+		{
+			UE_LOG_ONLINE_STOREV2(Warning,
+				TEXT("No logged in user."));
+		}
+	}
+	else
+	{
+		UE_LOG_ONLINE_STOREV2(Warning,
+			TEXT("Failed to get OSS interface for %s."), *SubsystemName);
+	}
+	// finish due to failure
+	FinishTest();
+}
+
+void FTestStoreInterface::StartNextTest()
+{
+	if (RequestOfferIds.Num() > 0)
+	{
+		OnlineSub->GetStoreV2Interface()->QueryOffersById(*LocalUserId, RequestOfferIds,
+			FOnQueryOnlineStoreOffersComplete::CreateRaw(this, &FTestStoreInterface::OnQueryOnlineStoreOffersComplete));
+	}
+	else
+	{
+		OnlineSub->GetStoreV2Interface()->QueryOffersByFilter(*LocalUserId, FOnlineStoreFilter(),
+			FOnQueryOnlineStoreOffersComplete::CreateRaw(this, &FTestStoreInterface::OnQueryOnlineStoreOffersComplete));
+	}
+}
+
+void FTestStoreInterface::OnQueryOnlineStoreOffersComplete(bool bWasSuccessful, const TArray<FUniqueOfferId>& OfferIds, const FString& ErrorString)
+{
+	UE_LOG_ONLINE_STOREV2(Log, TEXT("Query store offers completed with status bWasSuccessful=[%d] Error=[%s] OfferIds=[%d]"), bWasSuccessful, ErrorString.IsEmpty() ? TEXT("None") : *ErrorString, OfferIds.Num());
+
+	if (bWasSuccessful)
+	{
+		TArray<FOnlineStoreOfferRef> StoreOffers;
+		OnlineSub->GetStoreV2Interface()->GetOffers(StoreOffers);
+
+		TArray<FOnlineStoreOfferRef> FoundOffers;
+		UE_LOG(LogOnline, Log, TEXT("Found %d offers in cache"), StoreOffers.Num());
+		for (const FOnlineStoreOfferRef& OfferRef : StoreOffers)
+		{
+			UE_LOG(LogOnline, Log, TEXT("  Title=[%s] Desc=[%s] OfferId=[%s] CurrencyCode=[%s] PriceInt=[%d] DisplayPrice=[%s]"), 
+				*OfferRef->Title.ToString(),
+				*OfferRef->Description.ToString(),
+				*OfferRef->OfferId, 
+				*OfferRef->CurrencyCode, 
+				OfferRef->NumericPrice, 
+				*OfferRef->GetDisplayPrice().ToString()
+			);
+			if (RequestOfferIds.Contains(OfferRef->OfferId))
+			{
+				FoundOffers.Add(OfferRef);
+			}
+		}
+		UE_LOG(LogOnline, Log, TEXT("Requested %d offers and found %d offers"), RequestOfferIds.Num(), FoundOffers.Num());
+	}
+
+	FinishTest();
+};
+
+void FTestStoreInterface::FinishTest()
+{
+	delete this;
+}
+
+#endif //WITH_DEV_AUTOMATION_TESTS
diff --git a/Engine/Plugins/Online/OnlineSubsystemUtils/Source/OnlineSubsystemUtils/Private/Tests/TestStoreInterface.h b/Engine/Plugins/Online/OnlineSubsystemUtils/Source/OnlineSubsystemUtils/Private/Tests/TestStoreInterface.h
new file mode 100644
--- /dev/null
+++ b/Engine/Plugins/Online/OnlineSubsystemUtils/Source/OnlineSubsystemUtils/Private/Tests/TestStoreInterface.h
@@ -0,0 +1,69 @@
+// Copyright Epic Games, Inc. All Rights Reserved.
+
+#pragma once
+
+// UE4 includes
+
+#include "CoreMinimal.h"
+#include "UObject/CoreOnline.h"
+
+// Module includes
+#include "Interfaces/OnlineStoreInterfaceV2.h"
+
+class IOnlineSubsystem;
+
+#if WITH_DEV_AUTOMATION_TESTS
+
+/**
+ * Class used to test the Store interface
+ */
+class FTestStoreInterface
+{
+
+private:
+
+	/** The subsystem that was requested to be tested or the default if empty */
+	FString SubsystemName;
+	/** The online interface to use for testing */
+	IOnlineSubsystem* OnlineSub;
+	/** Logged in user */
+	TSharedPtr<const FUniqueNetId> LocalUserId;
+	/** Store offers to query */
+	TArray<FUniqueOfferId> RequestOfferIds;
+	
+	/**
+	 * Step through the various tests that should be run and initiate the next one
+	 */
+	void StartNextTest();
+
+	/**
+	 * Finish/cleanup the tests
+	 */
+	void FinishTest();
+
+	/** Callback for store query */
+	void OnQueryOnlineStoreOffersComplete(bool bWasSuccessful, const TArray<FUniqueOfferId>& OfferIds, const FString& ErrorString);
+
+public:
+
+	/**
+	 *	Constructor which sets the subsystem name to test
+	 *
+	 * @param InSubsystem the subsystem to test
+	 */
+	FTestStoreInterface(const FString& InSubsystem);
+
+	/**
+	 *	Destructor
+	 */
+	~FTestStoreInterface();
+
+	/**
+	 *	Kicks off all of the testing process
+	 *
+	 * @param InOffersIds list of product ids to query
+	 */
+	void Test(class UWorld* InWorld, const TArray<FString>& InOffersIds);
+ };
+
+#endif //WITH_DEV_AUTOMATION_TESTS
