Implement FindSessionById in TestSessionInterface

diff --git a/Engine/Plugins/Online/OnlineSubsystemUtils/Source/OnlineSubsystemUtils/Private/Tests/TestSessionInterface.cpp b/Engine/Plugins/Online/OnlineSubsystemUtils/Source/OnlineSubsystemUtils/Private/Tests/TestSessionInterface.cpp
--- a/Engine/Plugins/Online/OnlineSubsystemUtils/Source/OnlineSubsystemUtils/Private/Tests/TestSessionInterface.cpp
+++ b/Engine/Plugins/Online/OnlineSubsystemUtils/Source/OnlineSubsystemUtils/Private/Tests/TestSessionInterface.cpp
@@ -403,6 +403,16 @@ void FTestSessionInterface::OnFindSessionsComplete(bool bWasSuccessful)
 	}
 }
 
+void FTestSessionInterface::OnFindSessionByIdComplete(int32 LocalUserNum, bool bWasSuccessful, const FOnlineSessionSearchResult& SearchResult)
+{
+	UE_LOG_ONLINE_SESSION(Verbose, TEXT("OnFindSessionByIdComplete bSuccess: %d"), bWasSuccessful);
+
+	DumpSession(&SearchResult.Session);
+
+	SearchSettings = MakeShared<TestOnlineSearchSettings>(false, false);
+	SearchSettings->SearchResults.Emplace(SearchResult);
+}
+
 void FTestSessionInterface::OnCancelFindSessionsComplete(bool bWasSuccessful)
 {
 	UE_LOG_ONLINE_SESSION(Verbose, TEXT("OnCancelFindSessionsComplete bSuccess: %d"), bWasSuccessful);
@@ -473,13 +483,30 @@ bool FTestSessionInterface::Exec(UWorld* InWorld, const TCHAR* Cmd, FOutputDevic
 			}
 			bWasHandled = true;
 		}
+		else if (FParse::Command(&Cmd, TEXT("FINDBYID")))
+		{
+			FString SessionIdStr;
+			if (FParse::Token(Cmd, SessionIdStr, true))
+			{
+				TSharedPtr<const FUniqueNetId> SessionId = SessionInt->CreateSessionIdFromString(SessionIdStr);
+				if (SessionId.IsValid())
+				{
+					FOnSingleSessionResultCompleteDelegate CompletionDelegate;
+					CompletionDelegate.BindRaw(this, &FTestSessionInterface::OnFindSessionByIdComplete);
+					TSharedPtr<const FUniqueNetId> LocalUserId = Identity->GetUniquePlayerId(LocalUserNum);
+					check(LocalUserId.IsValid());
+					SessionInt->FindSessionById(*LocalUserId, *SessionId, *LocalUserId, CompletionDelegate);
+				}
+			}
+			bWasHandled = true;
+		}
 		else if (FParse::Command(&Cmd, TEXT("JOIN")))
 		{
 			TCHAR SearchIdxStr[256];
 			if (FParse::Token(Cmd, SearchIdxStr, UE_ARRAY_COUNT(SearchIdxStr), true))
 			{
 				int32 SearchIdx = FCString::Atoi(SearchIdxStr);
-				if (SearchIdx >= 0 && SearchIdx < SearchSettings->SearchResults.Num())
+				if (SearchSettings->SearchResults.IsValidIndex(SearchIdx))
 				{
 					JoinSession(LocalUserNum, SessionName, SearchSettings->SearchResults[SearchIdx]);
 				}
diff --git a/Engine/Plugins/Online/OnlineSubsystemUtils/Source/OnlineSubsystemUtils/Private/Tests/TestSessionInterface.h b/Engine/Plugins/Online/OnlineSubsystemUtils/Source/OnlineSubsystemUtils/Private/Tests/TestSessionInterface.h
--- a/Engine/Plugins/Online/OnlineSubsystemUtils/Source/OnlineSubsystemUtils/Private/Tests/TestSessionInterface.h
+++ b/Engine/Plugins/Online/OnlineSubsystemUtils/Source/OnlineSubsystemUtils/Private/Tests/TestSessionInterface.h
@@ -268,6 +268,15 @@ class TestOnlineSearchSettings;
 	 */
 	void OnFindSessionsComplete(bool bWasSuccessful);
 
+	/**
+	 * Delegate fired when the search for an online session has completed
+	 *
+	 * @param LocalUserNum the controller number of the accepting user
+	 * @param bWasSuccessful the session was found and is joinable, false otherwise
+	 * @param SearchResult the search/settings for the session result we've been given
+	 */
+	void OnFindSessionByIdComplete(int32 LocalUserNum, bool bWasSuccessful, const FOnlineSessionSearchResult& SearchResult);
+
 	/**
 	 * Delegate fired when the cancellation of a search for an online session has completed
 	 *
