LEADERBOARDS test refactor to admit additional parameters when querying leaderboards

diff --git a/Engine/Plugins/Online/OnlineSubsystem/Source/Public/OnlineKeyValuePair.h b/Engine/Plugins/Online/OnlineSubsystem/Source/Public/OnlineKeyValuePair.h
--- a/Engine/Plugins/Online/OnlineSubsystem/Source/Public/OnlineKeyValuePair.h
+++ b/Engine/Plugins/Online/OnlineSubsystem/Source/Public/OnlineKeyValuePair.h
@@ -64,6 +64,59 @@ namespace EOnlineKeyValuePairDataType
 			return TEXT("");
 		}		
 	}
+
+	/** @return the correspondent enum for the stringified version passed in */
+	inline EOnlineKeyValuePairDataType::Type FromString(const FString& EnumStr)
+	{
+		if (EnumStr.Equals(TEXT("Empty"),ESearchCase::IgnoreCase))
+		{
+			return Empty;
+		}
+		else if (EnumStr.Equals(TEXT("Int32"),ESearchCase::IgnoreCase))
+		{
+			return Int32;
+		}
+		else if (EnumStr.Equals(TEXT("UInt32"),ESearchCase::IgnoreCase))
+		{
+			return UInt32;
+		}
+		else if (EnumStr.Equals(TEXT("Int64"),ESearchCase::IgnoreCase))
+		{
+			return Int64;
+		}
+		else if (EnumStr.Equals(TEXT("UInt64"),ESearchCase::IgnoreCase))
+		{
+			return UInt64;
+		}
+		else if (EnumStr.Equals(TEXT("Double"),ESearchCase::IgnoreCase))
+		{
+			return Double;
+		}
+		else if (EnumStr.Equals(TEXT("String"),ESearchCase::IgnoreCase))
+		{
+			return String;
+		}
+		else if (EnumStr.Equals(TEXT("Float"),ESearchCase::IgnoreCase))
+		{
+			return Float;
+		}
+		else if (EnumStr.Equals(TEXT("Blob"),ESearchCase::IgnoreCase))
+		{
+			return Blob;
+		}
+		else if (EnumStr.Equals(TEXT("Bool"),ESearchCase::IgnoreCase))
+		{
+			return Bool;
+		}
+		else if (EnumStr.Equals(TEXT("Json"),ESearchCase::IgnoreCase))
+		{
+			return Json;
+		}
+		else
+		{
+			return MAX;
+		}
+	}
 }
 
 /**
diff --git a/Engine/Plugins/Online/OnlineSubsystemUtils/Source/OnlineSubsystemUtils/Private/OnlineSubsystemUtils.cpp b/Engine/Plugins/Online/OnlineSubsystemUtils/Source/OnlineSubsystemUtils/Private/OnlineSubsystemUtils.cpp
--- a/Engine/Plugins/Online/OnlineSubsystemUtils/Source/OnlineSubsystemUtils/Private/OnlineSubsystemUtils.cpp
+++ b/Engine/Plugins/Online/OnlineSubsystemUtils/Source/OnlineSubsystemUtils/Private/OnlineSubsystemUtils.cpp
@@ -600,7 +600,30 @@ static bool OnlineExec( UWorld* InWorld, const TCHAR* Cmd, FOutputDevice& Ar )
 					else if (FParse::Command(&Cmd, TEXT("LEADERBOARDS")))
 					{
 						// This class deletes itself once done
-						(new FTestLeaderboardInterface(SubName))->Test(InWorld, FParse::Token(Cmd, false));
+						FString LeaderboardName, SortedColumn, UserId;
+						
+						if (!FParse::Token(Cmd, LeaderboardName, false) || !FParse::Token(Cmd, SortedColumn, false))
+						{
+							UE_LOG_ONLINE_LEADERBOARD(Log, TEXT("Command parameters not found. Command syntax is 'LEADERBOARDS LeaderboardName SortedColumn ColumnNName ColumnNFormat ... UserID"));
+						}
+
+						TMap<FString, EOnlineKeyValuePairDataType::Type> Columns;
+						FString ColumnName;
+						while (FParse::Token(Cmd, ColumnName, false))
+						{
+							FString ColumnFormat;
+							if(FParse::Token(Cmd, ColumnFormat, false))
+							{
+								Columns.Add(ColumnName, EOnlineKeyValuePairDataType::FromString(ColumnFormat));
+							}
+							else
+							{
+								UE_LOG_ONLINE_LEADERBOARD(Log, TEXT("Setting %s as UserId for LEADERBOARDS TEST"), *ColumnName);
+								UserId = ColumnName;
+							}
+						}
+
+						(new FTestLeaderboardInterface(SubName))->Test(InWorld, LeaderboardName, SortedColumn, MoveTemp(Columns), UserId);
 						bWasHandled = true;
 					}
 					else if (FParse::Command(&Cmd, TEXT("PRESENCE")))
diff --git a/Engine/Plugins/Online/OnlineSubsystemUtils/Source/OnlineSubsystemUtils/Private/Tests/TestLeaderboardInterface.cpp b/Engine/Plugins/Online/OnlineSubsystemUtils/Source/OnlineSubsystemUtils/Private/Tests/TestLeaderboardInterface.cpp
--- a/Engine/Plugins/Online/OnlineSubsystemUtils/Source/OnlineSubsystemUtils/Private/Tests/TestLeaderboardInterface.cpp
+++ b/Engine/Plugins/Online/OnlineSubsystemUtils/Source/OnlineSubsystemUtils/Private/Tests/TestLeaderboardInterface.cpp
@@ -29,21 +29,22 @@ public:
 class TestLeaderboardRead : public FOnlineLeaderboardRead
 {
 public:
-	TestLeaderboardRead()
+	TestLeaderboardRead(const FString& InLeaderboardName, const FString& InSortedColumn, const TMap<FString, EOnlineKeyValuePairDataType::Type>& InColumns)
 	{
-		// Default properties
-		LeaderboardName = FName(TEXT("TestLeaderboard"));
-		SortedColumn = "TestIntStat1";
+		LeaderboardName = FName(InLeaderboardName);
+		SortedColumn = FName(InSortedColumn);
 
-		// Define default columns
-		new (ColumnMetadata) FColumnMetaData("TestIntStat1", EOnlineKeyValuePairDataType::Int32);
-		new (ColumnMetadata) FColumnMetaData("TestFloatStat1", EOnlineKeyValuePairDataType::Float);
+		for (TPair<FString, EOnlineKeyValuePairDataType::Type> Column : InColumns)
+		{
+			new (ColumnMetadata) FColumnMetaData(FName(Column.Key), Column.Value);
+		}
 	}
 };
 
 FTestLeaderboardInterface::FTestLeaderboardInterface(const FString& InSubsystem) :
 	Subsystem(InSubsystem),
 	bOverallSuccess(true),
+	bReadLeaderboardAttempted(false),
 	Leaderboards(NULL),
 	TestPhase(0),
 	LastTestPhase(-1)
@@ -68,7 +69,7 @@ FTestLeaderboardInterface::~FTestLeaderboardInterface()
 	Leaderboards = NULL;
 }
 
-void FTestLeaderboardInterface::Test(UWorld* InWorld, const FString& InUserId)
+void FTestLeaderboardInterface::Test(UWorld* InWorld, const FString& InLeaderboardName, const FString& InColumnName, TMap<FString, EOnlineKeyValuePairDataType::Type>&& InColumns, const FString& InUserId)
 {
 	FindRankUserId = InUserId;
 	OnlineSub = Online::GetSubsystem(InWorld, FName(*Subsystem));
@@ -94,6 +95,10 @@ void FTestLeaderboardInterface::Test(UWorld* InWorld, const FString& InUserId)
 		bOverallSuccess = false;
 		return;
 	}
+
+	LeaderboardName = InLeaderboardName;
+	SortedColumn = InColumnName;
+	Columns = MoveTemp(InColumns);
 }
 
 void FTestLeaderboardInterface::WriteLeaderboards()
@@ -141,58 +146,55 @@ void FTestLeaderboardInterface::PrintLeaderboards()
 void FTestLeaderboardInterface::OnLeaderboardReadComplete(bool bWasSuccessful)
 {
 	UE_LOG_ONLINE_LEADERBOARD(Verbose, TEXT("OnLeaderboardReadComplete bWasSuccessful: %d"), bWasSuccessful);
-	bOverallSuccess = bOverallSuccess && bWasSuccessful;
+	bOverallSuccess = bOverallSuccess && (bReadLeaderboardAttempted == bWasSuccessful);
 
 	PrintLeaderboards();
 
 	Leaderboards->ClearOnLeaderboardReadCompleteDelegate_Handle(LeaderboardReadCompleteDelegateHandle);
+	bReadLeaderboardAttempted = false;
 	TestPhase++;
 }
 
 void FTestLeaderboardInterface::OnLeaderboardRankReadComplete(bool bWasSuccessful)
 {
 	UE_LOG_ONLINE_LEADERBOARD(Verbose, TEXT("OnLeaderboardRankReadComplete bWasSuccessful: %d"), bWasSuccessful);
-	bOverallSuccess = bOverallSuccess && bWasSuccessful;
+	bOverallSuccess = bOverallSuccess && (bReadLeaderboardAttempted == bWasSuccessful);
 
 	PrintLeaderboards();
 
 	Leaderboards->ClearOnLeaderboardReadCompleteDelegate_Handle(LeaderboardReadRankCompleteDelegateHandle);
+	bReadLeaderboardAttempted = false;
 	TestPhase++;
 }
 
 void FTestLeaderboardInterface::OnLeaderboardUserRankReadComplete(bool bWasSuccessful)
 {
 	UE_LOG_ONLINE_LEADERBOARD(Verbose, TEXT("OnLeaderboardUserRankReadComplete bWasSuccessful: %d"), bWasSuccessful);
-	bOverallSuccess = bOverallSuccess && bWasSuccessful;
+	bOverallSuccess = bOverallSuccess && (bReadLeaderboardAttempted == bWasSuccessful);
 
 	PrintLeaderboards();
 
 	Leaderboards->ClearOnLeaderboardReadCompleteDelegate_Handle(LeaderboardReadRankUserCompleteDelegateHandle);
+	bReadLeaderboardAttempted = false;
 	TestPhase++;
 }
 
 void FTestLeaderboardInterface::ReadLeaderboards()
 {
-	ReadObject = MakeShareable(new TestLeaderboardRead());
+	ReadObject = MakeShareable(new TestLeaderboardRead(LeaderboardName, SortedColumn, Columns));
 	FOnlineLeaderboardReadRef ReadObjectRef = ReadObject.ToSharedRef();
 
 	LeaderboardReadCompleteDelegateHandle = Leaderboards->AddOnLeaderboardReadCompleteDelegate_Handle(LeaderboardReadCompleteDelegate);
-	Leaderboards->ReadLeaderboardsForFriends(0, ReadObjectRef);
+	bReadLeaderboardAttempted = Leaderboards->ReadLeaderboardsForFriends(0, ReadObjectRef);
 }
 
 void FTestLeaderboardInterface::ReadLeaderboardsRank(int32 Rank, int32 Range)
 {
-	ReadObject = MakeShareable(new TestLeaderboardRead());
+	ReadObject = MakeShareable(new TestLeaderboardRead(LeaderboardName, SortedColumn, Columns));
 	FOnlineLeaderboardReadRef ReadObjectRef = ReadObject.ToSharedRef();
 
 	LeaderboardReadRankCompleteDelegateHandle = Leaderboards->AddOnLeaderboardReadCompleteDelegate_Handle(LeaderboardReadRankCompleteDelegate);
-	if (!Leaderboards->ReadLeaderboardsAroundRank(Rank, Range, ReadObjectRef))
-	{
-		UE_LOG_ONLINE_LEADERBOARD(Warning, TEXT("Cannot run the leaderboards around rank test as it failed to start"));
-		bOverallSuccess = false;
-		Leaderboards->ClearOnLeaderboardReadCompleteDelegate_Handle(LeaderboardReadRankCompleteDelegateHandle);
-		++TestPhase;
-	}
+	bReadLeaderboardAttempted = Leaderboards->ReadLeaderboardsAroundRank(Rank, Range, ReadObjectRef);
 }
 
 void FTestLeaderboardInterface::ReadLeaderboardsUser(const FUniqueNetId& InUserId, int32 Range)
@@ -204,16 +206,20 @@ void FTestLeaderboardInterface::ReadLeaderboardsUser(const FUniqueNetId& InUserI
 		return;
 	}
 
-	ReadObject = MakeShareable(new TestLeaderboardRead());
+	ReadObject = MakeShareable(new TestLeaderboardRead(LeaderboardName, SortedColumn, Columns));
 	FOnlineLeaderboardReadRef ReadObjectRef = ReadObject.ToSharedRef();
 
 	// Need to get a shared reference for ReadLeaderboardsAroundUser
 	TSharedPtr<const FUniqueNetId> ArbitraryId = OnlineSub->GetIdentityInterface()->CreateUniquePlayerId(InUserId.ToString());
 
-	LeaderboardReadRankUserCompleteDelegateHandle = Leaderboards->AddOnLeaderboardReadCompleteDelegate_Handle(LeaderboardReadRankUserCompleteDelegate);
-	if (!ArbitraryId.IsValid() || !Leaderboards->ReadLeaderboardsAroundUser(ArbitraryId.ToSharedRef(), Range, ReadObjectRef))
+	if (ArbitraryId.IsValid())
+	{
+		LeaderboardReadRankUserCompleteDelegateHandle = Leaderboards->AddOnLeaderboardReadCompleteDelegate_Handle(LeaderboardReadRankUserCompleteDelegate);
+		bReadLeaderboardAttempted = Leaderboards->ReadLeaderboardsAroundUser(ArbitraryId.ToSharedRef(), Range, ReadObjectRef);
+	}
+	else
 	{
-		UE_LOG_ONLINE_LEADERBOARD(Warning, TEXT("Cannot run the leaderboards around user test as it failed to start"));
+		UE_LOG_ONLINE_LEADERBOARD(Warning, TEXT("Cannot run the leaderboards around user test as it failed to start. UserId not valid"));
 		bOverallSuccess = false;
 		Leaderboards->ClearOnLeaderboardReadCompleteDelegate_Handle(LeaderboardReadRankUserCompleteDelegateHandle);
 		++TestPhase;
diff --git a/Engine/Plugins/Online/OnlineSubsystemUtils/Source/OnlineSubsystemUtils/Private/Tests/TestLeaderboardInterface.h b/Engine/Plugins/Online/OnlineSubsystemUtils/Source/OnlineSubsystemUtils/Private/Tests/TestLeaderboardInterface.h
--- a/Engine/Plugins/Online/OnlineSubsystemUtils/Source/OnlineSubsystemUtils/Private/Tests/TestLeaderboardInterface.h
+++ b/Engine/Plugins/Online/OnlineSubsystemUtils/Source/OnlineSubsystemUtils/Private/Tests/TestLeaderboardInterface.h
@@ -23,12 +23,24 @@
 	/** Keep track of success across all functions and callbacks */
 	bool bOverallSuccess;
 
+	/** Set every time we try to read a Leaderboard, indicates if the read was attempted or not */
+	bool bReadLeaderboardAttempted;
+
 	/** Logged in UserId */
 	TSharedPtr<const FUniqueNetId> UserId;
 
 	/** Passed in UserId */
 	FString FindRankUserId;
 
+	/** Passed in LeaderboardName */
+	FString LeaderboardName;
+
+	/** Passed in ColumnName */
+	FString SortedColumn;
+
+	/** Passed in Columns */
+	TMap<FString, EOnlineKeyValuePairDataType::Type> Columns;
+
 	/** Convenient access to the leaderboard interfaces */
 	IOnlineLeaderboardsPtr Leaderboards;
 
@@ -111,7 +123,7 @@
 	/**
 	 * Kicks off all of the testing process
 	 */
-	void Test(class UWorld* InWorld, const FString& InUserId);
+	void Test(UWorld* InWorld, const FString& InLeaderboardName, const FString& InSortedColumn, TMap<FString, EOnlineKeyValuePairDataType::Type>&& InColumns, const FString& InUserId);
  };
 
 #endif //WITH_DEV_AUTOMATION_TESTS
