Modify TestSessionClient to not require a parameter to JoinSession
#rnx

diff --git a/Engine/Plugins/Online/OnlineSubsystemUtils/Source/OnlineSubsystemUtils/Private/Tests/TestSessionInterface.cpp b/Engine/Plugins/Online/OnlineSubsystemUtils/Source/OnlineSubsystemUtils/Private/Tests/TestSessionInterface.cpp
--- a/Engine/Plugins/Online/OnlineSubsystemUtils/Source/OnlineSubsystemUtils/Private/Tests/TestSessionInterface.cpp
+++ b/Engine/Plugins/Online/OnlineSubsystemUtils/Source/OnlineSubsystemUtils/Private/Tests/TestSessionInterface.cpp
@@ -502,14 +502,15 @@ bool FTestSessionInterface::Exec(UWorld* InWorld, const TCHAR* Cmd, FOutputDevic
 		}
 		else if (FParse::Command(&Cmd, TEXT("JOIN")))
 		{
+			int32 SearchIdx = 0;
 			TCHAR SearchIdxStr[256];
 			if (FParse::Token(Cmd, SearchIdxStr, UE_ARRAY_COUNT(SearchIdxStr), true))
 			{
-				int32 SearchIdx = FCString::Atoi(SearchIdxStr);
-				if (SearchSettings->SearchResults.IsValidIndex(SearchIdx))
-				{
-					JoinSession(LocalUserNum, SessionName, SearchSettings->SearchResults[SearchIdx]);
-				}
+				SearchIdx = FCString::Atoi(SearchIdxStr);
+			}
+			if (SearchSettings->SearchResults.IsValidIndex(SearchIdx))
+			{
+				JoinSession(LocalUserNum, SessionName, SearchSettings->SearchResults[SearchIdx]);
 			}
 			bWasHandled = true;
 		}
