Fix for crashes that can occur when destroying IpNetDrivers and IpConnections due to unsafe lifetime management of sockets.
Introduced more robust socket lifetime management. Added functions to return unique pointers to sockets from the SocketSubsystem.
IpNetDriver and IpConnection now manage socket lifetime through shared pointers except for where they can't due to backwards compatibility concerns.

diff --git a/Engine/Plugins/Online/OnlineSubsystemUtils/Source/OnlineSubsystemUtils/Classes/IpConnection.h b/Engine/Plugins/Online/OnlineSubsystemUtils/Source/OnlineSubsystemUtils/Classes/IpConnection.h
--- a/Engine/Plugins/Online/OnlineSubsystemUtils/Source/OnlineSubsystemUtils/Classes/IpConnection.h
+++ b/Engine/Plugins/Online/OnlineSubsystemUtils/Source/OnlineSubsystemUtils/Classes/IpConnection.h
@@ -35,6 +35,7 @@ class ONLINESUBSYSTEMUTILS_API UIpConnection : public UNetConnection
     GENERATED_UCLASS_BODY()
 	// Variables.
 
+	/** This is a non-owning pointer to a socket owned elsewhere, IpConnection will not destroy the socket through this pointer. */
 	class FSocket*				Socket;
 	UE_DEPRECATED(4.25, "Address resolution is now handled in the IpNetDriver and no longer done entirely in the IpConnection")
 	class FResolveInfo*			ResolveInfo;
@@ -106,7 +107,10 @@ private:
 	void HandleSocketRecvError(class UNetDriver* NetDriver, const FString& ErrorString);
 
 	/** An array of sockets tied to every binding address. */
-	TArray<FSocket*> BindSockets;
+	TArray<TSharedPtr<FSocket>> BindSockets;
+
+	/** Holds a refcount to the actual socket to be used from BindSockets. */
+	TSharedPtr<FSocket> ResolutionSocket;
 
 	/** An array containing the address results GAI returns for the current host value. Given to us from the netdriver. */
 	TArray<TSharedRef<FInternetAddr>> ResolverResults;
diff --git a/Engine/Plugins/Online/OnlineSubsystemUtils/Source/OnlineSubsystemUtils/Classes/IpNetDriver.h b/Engine/Plugins/Online/OnlineSubsystemUtils/Source/OnlineSubsystemUtils/Classes/IpNetDriver.h
--- a/Engine/Plugins/Online/OnlineSubsystemUtils/Source/OnlineSubsystemUtils/Classes/IpNetDriver.h
+++ b/Engine/Plugins/Online/OnlineSubsystemUtils/Source/OnlineSubsystemUtils/Classes/IpNetDriver.h
@@ -13,6 +13,7 @@
 #include "HAL/RunnableThread.h"
 #include "Containers/CircularQueue.h"
 #include "SocketTypes.h"
+#include "SocketSubsystem.h"
 #include "IpNetDriver.generated.h"
 
 
@@ -110,6 +111,9 @@ private:
 	 */
 	UNetConnection* ProcessConnectionlessPacket(FReceivedPacketView& PacketRef, const FPacketBufferView& WorkingBuffer);
 
+	/** Clears all references to sockets held by this net driver. */
+	void ClearSockets();
+
 public:
 	//~ Begin UIpNetDriver Interface.
 
@@ -118,6 +122,7 @@ public:
 	 *
 	 * @return an FSocket if creation succeeded, nullptr if creation failed.
 	 */
+	UE_DEPRECATED(4.25, "Socket creation is now restricted to subclasses and done through CreateSocketForProtocol or CreateAndBindSocket")
 	virtual FSocket* CreateSocket();
 	
 	/**
@@ -129,12 +134,22 @@ public:
 	virtual FSocket* GetSocket();
 
 	/**
-	 * Set the current NetDriver's socket to the given socket. This is typically done after resolution completes successfully. 
+	 * Set the current NetDriver's socket to the given socket and takes ownership of it (the driver will handle destroying it).
+	 * This is typically done after resolution completes successfully. 
 	 * This will also set the LocalAddr for the netdriver automatically.
 	 *
 	 * @param NewSocket the socket pointer to set this netdriver's socket to
 	 */
-	virtual void SetSocketAndLocalAddress(FSocket* NewSocket);
+	void SetSocketAndLocalAddress(FSocket* NewSocket);
+
+	/**
+	* Set the current NetDriver's socket to the given socket.
+	* This is typically done after resolution completes successfully. 
+	* This will also set the LocalAddr for the netdriver automatically.
+	*
+	* @param NewSocket the socket pointer to set this netdriver's socket to
+	*/
+	void SetSocketAndLocalAddress(const TSharedPtr<FSocket>& SharedSocket);
 
 	/**
 	 * Returns the port number to use when a client is creating a socket.
@@ -152,7 +167,7 @@ protected:
 	 * @param ProtocolType	an FName that represents the protocol to allocate the new socket under. Typically set to None or a value in FNetworkProtocolTypes
 	 * @return				an FSocket if creation succeeded, nullptr if creation failed.
 	 */
-	virtual FSocket* CreateSocketForProtocol(const FName& ProtocolType);
+	virtual FUniqueSocket CreateSocketForProtocol(const FName& ProtocolType);
 
 	/**
 	 * Creates, initializes and binds a socket using the given bind address information.
@@ -166,7 +181,7 @@ protected:
 	 *
 	 * @return if the socket could be created and bound with all the appropriate options, a pointer to the new socket is given, otherwise null
 	 */
-	virtual FSocket* CreateAndBindSocket(TSharedRef<FInternetAddr> BindAddr, int32 Port, bool bReuseAddressAndPort, int32 DesiredRecvSize, int32 DesiredSendSize, FString& Error);
+	virtual FUniqueSocket CreateAndBindSocket(TSharedRef<FInternetAddr> BindAddr, int32 Port, bool bReuseAddressAndPort, int32 DesiredRecvSize, int32 DesiredSendSize, FString& Error);
 	//~ End UIpNetDriver Interface.
 
 public:
@@ -299,5 +314,8 @@ private:
 	 * An array sockets created for every binding address a machine has in use for performing address resolution. 
 	 * This array empties after connections have been spun up.
 	 */
-	TArray<FSocket*> BoundSockets;
+	TArray<TSharedPtr<FSocket>> BoundSockets;
+
+	/** Underlying socket communication */
+	TSharedPtr<FSocket> SocketPrivate;
 };
diff --git a/Engine/Plugins/Online/OnlineSubsystemUtils/Source/OnlineSubsystemUtils/Private/IpConnection.cpp b/Engine/Plugins/Online/OnlineSubsystemUtils/Source/OnlineSubsystemUtils/Private/IpConnection.cpp
--- a/Engine/Plugins/Online/OnlineSubsystemUtils/Source/OnlineSubsystemUtils/Private/IpConnection.cpp
+++ b/Engine/Plugins/Online/OnlineSubsystemUtils/Source/OnlineSubsystemUtils/Private/IpConnection.cpp
@@ -185,11 +185,12 @@ void UIpConnection::Tick()
 		RemoteAddr = ResolverResults[CurrentAddressIndex];
 		Socket = nullptr;
 
-		for (FSocket* BindSocket : BindSockets)
+		for (const TSharedPtr<FSocket>& BindSocket : BindSockets)
 		{
 			if (BindSocket->GetProtocol() == RemoteAddr->GetProtocolType())
 			{
-				Socket = BindSocket;
+				ResolutionSocket = BindSocket;
+				Socket = ResolutionSocket.Get();
 				break;
 			}
 		}
@@ -230,7 +231,7 @@ void UIpConnection::Tick()
 		UIpNetDriver* IpDriver = Cast<UIpNetDriver>(Driver);
 
 		// Set the right object now that we have a connection
-		IpDriver->SetSocketAndLocalAddress(Socket);
+		IpDriver->SetSocketAndLocalAddress(ResolutionSocket);
 
 		ResolutionState = EAddressResolutionState::Done;
 	}
@@ -596,26 +597,17 @@ void UIpConnection::CleanupResolutionSockets()
 		return;
 	}
 
-	if (BindSockets.Num() > 0)
+	if (bCleanAll)
 	{
-		ISocketSubsystem* const SocketSubsystem = Driver->GetSocketSubsystem();
-		UE_LOG(LogNet, Verbose, TEXT("Cleaned up the sockets used for resolving!"));
-		// Clean up other sockets.
-		for (FSocket*& CurSocket : BindSockets)
+		if(ResolutionSocket.Get() == Socket)
 		{
-			if (CurSocket != Socket || bCleanAll)
-			{
-				SocketSubsystem->DestroySocket(CurSocket);
-				if (CurSocket == Socket)
-				{
-					Socket = nullptr;
-				}
-				CurSocket = nullptr;
-			}
+			Socket = nullptr;
 		}
-		BindSockets.Remove(nullptr);
+		
+		ResolutionSocket.Reset();
 	}
-
+	
+	BindSockets.Reset();
 	ResolverResults.Empty();
 }
 
diff --git a/Engine/Plugins/Online/OnlineSubsystemUtils/Source/OnlineSubsystemUtils/Private/IpNetDriver.cpp b/Engine/Plugins/Online/OnlineSubsystemUtils/Source/OnlineSubsystemUtils/Private/IpNetDriver.cpp
--- a/Engine/Plugins/Online/OnlineSubsystemUtils/Source/OnlineSubsystemUtils/Private/IpNetDriver.cpp
+++ b/Engine/Plugins/Online/OnlineSubsystemUtils/Source/OnlineSubsystemUtils/Private/IpNetDriver.cpp
@@ -655,7 +655,7 @@ private:
 		Connection->HandleSocketRecvError(Driver, ErrorString);
 	}
 
-	static void PushSocketsToConnection(UIpConnection* Connection, TArray<FSocket*>& Sockets)
+	static void PushSocketsToConnection(UIpConnection* Connection, TArray<TSharedPtr<FSocket>>& Sockets)
 	{
 		UE_LOG(LogNet, Verbose, TEXT("Pushed %d sockets to net connection %s"), Sockets.Num(), *Connection->GetName());
 		Connection->BindSockets = Sockets;
@@ -731,10 +731,11 @@ ISocketSubsystem* UIpNetDriver::GetSocketSubsystem()
 
 FSocket * UIpNetDriver::CreateSocket()
 {
-	return CreateSocketForProtocol( ( LocalAddr.IsValid() ? LocalAddr->GetProtocolType() : NAME_None) );
+	// This is a deprecated function with unsafe socket lifetime management. The Release call is intentional and for backwards compatiblity only.
+	return CreateSocketForProtocol((LocalAddr.IsValid() ? LocalAddr->GetProtocolType() : NAME_None)).Release();
 }
 
-FSocket* UIpNetDriver::CreateSocketForProtocol(const FName& ProtocolType)
+FUniqueSocket UIpNetDriver::CreateSocketForProtocol(const FName& ProtocolType)
 {
 	// Create UDP socket and enable broadcasting.
 	ISocketSubsystem* SocketSubsystem = GetSocketSubsystem();
@@ -745,7 +746,7 @@ FSocket* UIpNetDriver::CreateSocketForProtocol(const FName& ProtocolType)
 		return NULL;
 	}
 
-	return SocketSubsystem->CreateSocket(NAME_DGram, TEXT("Unreal"), ProtocolType);
+	return SocketSubsystem->CreateUniqueSocket(NAME_DGram, TEXT("Unreal"), ProtocolType);
 }
 
 int UIpNetDriver::GetClientPort()
@@ -753,9 +754,8 @@ int UIpNetDriver::GetClientPort()
 	return 0;
 }
 
-FSocket* UIpNetDriver::CreateAndBindSocket(TSharedRef<FInternetAddr> BindAddr, int32 Port, bool bReuseAddressAndPort, int32 DesiredRecvSize, int32 DesiredSendSize, FString& Error)
+FUniqueSocket UIpNetDriver::CreateAndBindSocket(TSharedRef<FInternetAddr> BindAddr, int32 Port, bool bReuseAddressAndPort, int32 DesiredRecvSize, int32 DesiredSendSize, FString& Error)
 {
-	FSocket* NewSocket = nullptr;
 	ISocketSubsystem* SocketSubsystem = GetSocketSubsystem();
 	if (SocketSubsystem == nullptr)
 	{
@@ -764,9 +764,9 @@ FSocket* UIpNetDriver::CreateAndBindSocket(TSharedRef<FInternetAddr> BindAddr, i
 	}
 
 	// Create the socket that we will use to communicate with
-	NewSocket = CreateSocketForProtocol(BindAddr->GetProtocolType());
+	FUniqueSocket NewSocket = CreateSocketForProtocol(BindAddr->GetProtocolType());
 
-	if (NewSocket == nullptr)
+	if (!NewSocket.IsValid())
 	{
 		Error = FString::Printf(TEXT("%s: socket failed (%i)"), SocketSubsystem->GetSocketAPIName(), (int32)SocketSubsystem->GetLastErrorCode());
 		return nullptr;
@@ -777,7 +777,7 @@ FSocket* UIpNetDriver::CreateAndBindSocket(TSharedRef<FInternetAddr> BindAddr, i
 	{
 		if (Error.IsEmpty() == false)
 		{
-			SocketSubsystem->DestroySocket(NewSocket);
+			NewSocket.Reset();
 		}
 	};
 
@@ -808,7 +808,7 @@ FSocket* UIpNetDriver::CreateAndBindSocket(TSharedRef<FInternetAddr> BindAddr, i
 	BindAddr->SetPort(Port);
 
 	int32 AttemptPort = BindAddr->GetPort();
-	int32 BoundPort = SocketSubsystem->BindNextPort(NewSocket, *BindAddr, MaxPortCountToTry + 1, 1);
+	int32 BoundPort = SocketSubsystem->BindNextPort(NewSocket.Get(), *BindAddr, MaxPortCountToTry + 1, 1);
 	if (BoundPort == 0)
 	{
 		Error = FString::Printf(TEXT("%s: binding to port %i failed (%i)"), SocketSubsystem->GetSocketAPIName(), AttemptPort,
@@ -827,10 +827,18 @@ FSocket* UIpNetDriver::CreateAndBindSocket(TSharedRef<FInternetAddr> BindAddr, i
 
 void UIpNetDriver::SetSocketAndLocalAddress(FSocket* NewSocket)
 {
+	SetSocketAndLocalAddress(TSharedPtr<FSocket>(NewSocket, FSocketDeleter(GetSocketSubsystem())));
+}
+
+void UIpNetDriver::SetSocketAndLocalAddress(const TSharedPtr<FSocket>& SharedSocket)
+{
+	SocketPrivate = SharedSocket;
+
 	PRAGMA_DISABLE_DEPRECATION_WARNINGS
-	Socket = NewSocket;
+	Socket = SocketPrivate.Get();
+	PRAGMA_ENABLE_DEPRECATION_WARNINGS
 
-	if (NewSocket != nullptr)
+	if (SocketPrivate.IsValid())
 	{
 		// Allocate any LocalAddrs if they haven't been allocated yet.
 		if (!LocalAddr.IsValid())
@@ -838,9 +846,28 @@ void UIpNetDriver::SetSocketAndLocalAddress(FSocket* NewSocket)
 			LocalAddr = GetSocketSubsystem()->CreateInternetAddr();
 		}
 
-		Socket->GetAddress(*LocalAddr);
+		SocketPrivate->GetAddress(*LocalAddr);
+	}
+}
+
+void UIpNetDriver::ClearSockets()
+{
+	// For backwards compatability with the public Socket member. Destroy it manually if it won't be destroyed by the reset below.
+	PRAGMA_DISABLE_DEPRECATION_WARNINGS
+	if(!ensureMsgf(Socket == SocketPrivate.Get(), TEXT("UIpNetDriver::ClearSockets: Socket and SocketPrivate point to different sockets! %s"), *GetDescription()))
+	{
+		ISocketSubsystem* const SocketSubsystem = GetSocketSubsystem();
+
+		if (SocketSubsystem)
+		{
+			SocketSubsystem->DestroySocket(Socket);
+		}
 	}
+	Socket = nullptr;
 	PRAGMA_ENABLE_DEPRECATION_WARNINGS
+
+	SocketPrivate.Reset();
+	BoundSockets.Reset();
 }
 
 bool UIpNetDriver::InitBase( bool bInitAsClient, FNetworkNotify* InNotify, const FURL& URL, bool bReuseAddressAndPort, FString& Error )
@@ -876,11 +903,11 @@ bool UIpNetDriver::InitBase( bool bInitAsClient, FNetworkNotify* InNotify, const
 	// Create sockets for every bind address
 	for (TSharedRef<FInternetAddr>& BindAddr : BindAddresses)
 	{
-		FSocket* NewSocket = CreateAndBindSocket(BindAddr, BindPort, bReuseAddressAndPort, DesiredRecvSize, DesiredSendSize, Error);
-		if (NewSocket != nullptr)
+		FUniqueSocket NewSocket = CreateAndBindSocket(BindAddr, BindPort, bReuseAddressAndPort, DesiredRecvSize, DesiredSendSize, Error);
+		if (NewSocket.IsValid())
 		{
 			UE_LOG(LogNet, Log, TEXT("Created socket for bind address: %s on port %d"), *BindAddr->ToString(false), BindPort);
-			BoundSockets.Add(NewSocket);
+			BoundSockets.Emplace(NewSocket.Release(), FSocketDeleter(NewSocket.GetDeleter()));
 		}
 		else
 		{
@@ -901,11 +928,7 @@ bool UIpNetDriver::InitBase( bool bInitAsClient, FNetworkNotify* InNotify, const
 		UE_LOG(LogNet, Warning, TEXT("Encountered an error while creating sockets for the bind addresses. %s"), *Error);
 		
 		// Make sure to destroy all sockets that we don't end up using.
-		for (auto& SubSocket : BoundSockets)
-		{
-			SocketSubsystem->DestroySocket(SubSocket);
-		}
-		BoundSockets.Empty();
+		BoundSockets.Reset();
 
 		return false;
 	}
@@ -945,7 +968,7 @@ bool UIpNetDriver::InitBase( bool bInitAsClient, FNetworkNotify* InNotify, const
 			if (bRetrieveTimestamps)
 			{
 				// Properly set this flag for every socket for each bind address.
-				for (auto& SubSocket : BoundSockets)
+				for (TSharedPtr<FSocket>& SubSocket : BoundSockets)
 				{
 					SubSocket->SetRetrieveTimestamp(true);
 				}
@@ -1065,17 +1088,10 @@ bool UIpNetDriver::InitConnect( FNetworkNotify* InNotify, const FURL& ConnectURL
 		// Clean up any potential multiple sockets we have created when resolution was disabled.
 		// InitBase could have created multiple sockets and if so, we'll want to clean them up.
 		UE_LOG(LogNet, Verbose, TEXT("Cleaning up additional sockets created as address resolution is disabled."));
-		for (FSocket*& CurSocket : BoundSockets)
+		BoundSockets.RemoveAll([InSocket = GetSocket()](const TSharedPtr<FSocket>& CurSocket)
 		{
-			PRAGMA_DISABLE_DEPRECATION_WARNINGS
-			if (CurSocket != Socket)
-			PRAGMA_ENABLE_DEPRECATION_WARNINGS
-			{
-				SocketSubsystem->DestroySocket(CurSocket);
-				CurSocket = nullptr;
-			}
-		}
-		BoundSockets.Remove(nullptr);
+			return CurSocket.Get() != InSocket;
+		});
 	}
 	
 	UE_LOG(LogNet, Log, TEXT("Game client on port %i, rate %i"), DestinationPort, ServerConnection->CurrentNetSpeed );
@@ -1644,18 +1660,12 @@ void UIpNetDriver::LowLevelDestroy()
 			UE_LOG(LogExit, Log, TEXT("closesocket error (%i)"), (int32)SocketSubsystem->GetLastErrorCode() );
 		}
 
-		PRAGMA_DISABLE_DEPRECATION_WARNINGS
 		if (FIpConnectionHelper::IsAddressResolutionEnabledForConnection(IpServerConnection))
 		{
 			FIpConnectionHelper::CleanUpConnectionSockets(IpServerConnection);
 		}
-		else
-		{
-			// Free the memory the OS allocated for this socket
-			SocketSubsystem->DestroySocket(Socket);
-		}
-		Socket = nullptr;
-		PRAGMA_ENABLE_DEPRECATION_WARNINGS
+
+		ClearSockets();
 
 		UE_LOG(LogExit, Log, TEXT("%s shut down"),*GetDescription() );
 	}
diff --git a/Engine/Source/Runtime/Sockets/Private/SocketSubsystem.cpp b/Engine/Source/Runtime/Sockets/Private/SocketSubsystem.cpp
--- a/Engine/Source/Runtime/Sockets/Private/SocketSubsystem.cpp
+++ b/Engine/Source/Runtime/Sockets/Private/SocketSubsystem.cpp
@@ -75,6 +75,18 @@ static IModuleInterface* LoadSubsystemModule(const FString& SubsystemName)
 	return nullptr;
 }
 
+
+FUniqueSocket ISocketSubsystem::CreateUniqueSocket(const FName& SocketType, const FString& SocketDescription, bool bForceUDP)
+{
+	return FUniqueSocket(CreateSocket(SocketType, SocketDescription, bForceUDP), FSocketDeleter(this));
+}
+
+FUniqueSocket ISocketSubsystem::CreateUniqueSocket(const FName& SocketType, const FString& SocketDescription, const FName& ProtocolName)
+{
+	return FUniqueSocket(CreateSocket(SocketType, SocketDescription, ProtocolName), FSocketDeleter(this));
+}
+
+
 /**
  * Shutdown all registered subsystems
  */
diff --git a/Engine/Source/Runtime/Sockets/Public/SocketSubsystem.h b/Engine/Source/Runtime/Sockets/Public/SocketSubsystem.h
--- a/Engine/Source/Runtime/Sockets/Public/SocketSubsystem.h
+++ b/Engine/Source/Runtime/Sockets/Public/SocketSubsystem.h
@@ -35,6 +35,10 @@ SOCKETS_API DECLARE_LOG_CATEGORY_EXTERN(LogSockets, Log, All);
 	#endif
 #endif
 
+/** A unique pointer that will automatically call ISocketSubsystem::DestroySocket on its owned socket. */
+class FSocketDeleter;
+using FUniqueSocket = TUniquePtr<FSocket, FSocketDeleter>;
+
 /**
  * This is the base interface to abstract platform specific sockets API
  * differences.
@@ -111,6 +115,30 @@ public:
 	 */
 	virtual FSocket* CreateSocket(const FName& SocketType, const FString& SocketDescription, const FName& ProtocolName) = 0;
 
+	/**
+	 * Creates a socket wrapped in a unique pointer that will call DestroySocket automatically - do not call it explicitly!
+	 * This SocketSubsystem must also outlive the sockets it creates.
+	 *
+	 * @Param SocketType type of socket to create (DGram, Stream, etc)
+	 * @param SocketDescription debug description
+	 * @param bForceUDP overrides any platform specific protocol with UDP instead
+	 *
+	 * @return the new socket or NULL if failed
+	 */
+	FUniqueSocket CreateUniqueSocket(const FName& SocketType, const FString& SocketDescription, bool bForceUDP = false);
+
+	/**
+	 * Creates a socket using the given protocol name,  wrapped in a unique pointer that will call DestroySocket automatically - do not call it explicitly!
+	 * This SocketSubsystem must also outlive the sockets it creates.
+	 *
+	 * @Param SocketType type of socket to create (DGram, Stream, etc)
+	 * @param SocketDescription debug description
+	 * @param ProtocolName the name of the internet protocol to use for this socket. None should be handled.
+	 *
+	 * @return the new socket or NULL if failed
+	 */
+	FUniqueSocket CreateUniqueSocket(const FName& SocketType, const FString& SocketDescription, const FName& ProtocolName);
+
 	/**
 	 * Creates a resolve info cached struct to hold the resolved address
 	 *
@@ -469,6 +497,32 @@ private:
 	TMap<FString, TSharedPtr<FInternetAddr> > HostNameCache;
 };
 
+/** Deleter object that can be used with unique & shared pointers that store FSockets. The SocketSubsystem must be valid when the call operator is invoked! */
+class FSocketDeleter
+{
+public:
+	FSocketDeleter()
+		: Subsystem(nullptr)
+	{
+	}
+
+	FSocketDeleter(ISocketSubsystem* InSubsystem)
+		: Subsystem(InSubsystem)
+	{
+	}
+
+	void operator()(FSocket* Socket) const
+	{
+		if (Subsystem && Socket)
+		{
+			Subsystem->DestroySocket(Socket);
+		}
+	}
+
+private:
+	ISocketSubsystem* Subsystem;
+};
+
 typedef TSharedPtr<ISocketSubsystem> IOnlineSocketPtr;
 
 
