Friends test refactor to remove incorrect references and improve error handling.

diff --git a/Engine/Plugins/Online/OnlineSubsystemUtils/Source/OnlineSubsystemUtils/Private/Tests/TestFriendsInterface.cpp b/Engine/Plugins/Online/OnlineSubsystemUtils/Source/OnlineSubsystemUtils/Private/Tests/TestFriendsInterface.cpp
--- a/Engine/Plugins/Online/OnlineSubsystemUtils/Source/OnlineSubsystemUtils/Private/Tests/TestFriendsInterface.cpp
+++ b/Engine/Plugins/Online/OnlineSubsystemUtils/Source/OnlineSubsystemUtils/Private/Tests/TestFriendsInterface.cpp
@@ -90,22 +90,10 @@ void FTestFriendsInterface::StartNextTest()
 	{
 		UE_LOG_ONLINE_FRIEND(Log, TEXT("Querying recent players"));
 
-		//Anvil doesn't have methods to query recent players in the gdk
-		if (SubsystemName.Equals("anvil", ESearchCase::IgnoreCase))
+		if (OnlineSub->GetIdentityInterface().IsValid() &&
+			OnlineSub->GetIdentityInterface()->GetUniquePlayerId(0).IsValid())
 		{
-			UE_LOG_ONLINE_FRIEND(Log, TEXT("The Anvil subsystem doesn't have methods to query recent players. Skipping test"));
-
-			bQueryRecentPlayers = false;
-			StartNextTest();
-		}
-		else
-		{
-			if (OnlineSub->GetIdentityInterface().IsValid() &&
-				OnlineSub->GetIdentityInterface()->GetUniquePlayerId(0).IsValid())
-			{
-				OnlineSub->GetFriendsInterface()->QueryRecentPlayers(*OnlineSub->GetIdentityInterface()->GetUniquePlayerId(0), RecentPlayersNamespace);
-			}
-			bQueryRecentPlayers = false;
+			OnlineSub->GetFriendsInterface()->QueryRecentPlayers(*OnlineSub->GetIdentityInterface()->GetUniquePlayerId(0), RecentPlayersNamespace);
 		}
 	}
 	else if (bAcceptInvites && InvitesToAccept.Num() > 0)
@@ -126,36 +114,14 @@ void FTestFriendsInterface::StartNextTest()
 	{
 		UE_LOG_ONLINE_FRIEND(Log, TEXT("Deleting friends"));
 
-		//Anvil doesn't have methods to remove friends in the gdk
-		if (SubsystemName.Equals("anvil", ESearchCase::IgnoreCase))
-		{
-			UE_LOG_ONLINE_FRIEND(Log, TEXT("The Anvil subsystem doesn't have methods to delete friends. Skipping test"));
-
-			bDeleteFriends = false;
-			StartNextTest();
-		}
-		else
-		{
-			OnlineSub->GetFriendsInterface()->DeleteFriend(0, *FriendsToDelete[0], FriendsListName);
-		}
+		OnlineSub->GetFriendsInterface()->DeleteFriend(0, *FriendsToDelete[0], FriendsListName);
 	}
 	else if (bDeleteFriendsList)
 	{
 		UE_LOG_ONLINE_FRIEND(Log, TEXT("Deleting friends list"));
 
-		//Anvil doesn't have methods to remove friend lists in the gdk
-		if (SubsystemName.Equals("anvil", ESearchCase::IgnoreCase))
-		{
-			UE_LOG_ONLINE_FRIEND(Log, TEXT("The Anvil subsystem doesn't have methods to delete friends lists. Skipping test"));
-
-			bDeleteFriendsList = false;
-			StartNextTest();
-		}
-		else
-		{
-			FOnDeleteFriendsListComplete Delegate = FOnDeleteFriendsListComplete::CreateRaw(this, &FTestFriendsInterface::OnDeleteFriendsListComplete);
-			OnlineSub->GetFriendsInterface()->DeleteFriendsList(0, FriendsListName, Delegate);
-		}
+		FOnDeleteFriendsListComplete Delegate = FOnDeleteFriendsListComplete::CreateRaw(this, &FTestFriendsInterface::OnDeleteFriendsListComplete);
+		OnlineSub->GetFriendsInterface()->DeleteFriendsList(0, FriendsListName, Delegate);
 	}
 	else if (bQueryBlockedPlayers)
 	{
