Deprecate PlayTogether as it will be removed from the SDK (OSS)

diff --git a/Engine/Plugins/Online/OnlineSubsystem/Source/Public/OnlineSubsystem.h b/Engine/Plugins/Online/OnlineSubsystem/Source/Public/OnlineSubsystem.h
--- a/Engine/Plugins/Online/OnlineSubsystem/Source/Public/OnlineSubsystem.h
+++ b/Engine/Plugins/Online/OnlineSubsystem/Source/Public/OnlineSubsystem.h
@@ -127,7 +127,9 @@ typedef FOnOnlineEnvironmentChanged::FDelegate FOnOnlineEnvironmentChangedDelega
 * @param UserIndex - User index of the player the event is for
 * @param UserIdList - list of other users in the PS4 party to send invites to
 */
+UE_DEPRECATED(4.26, "PlayTogether will no longer be supported and should be removed.")
 DECLARE_MULTICAST_DELEGATE_TwoParams(FOnPlayTogetherEventReceived, int32, TArray<TSharedPtr<const FUniqueNetId>>);
+UE_DEPRECATED(4.26, "PlayTogether will no longer be supported and should be removed.")
 typedef FOnPlayTogetherEventReceived::FDelegate FOnPlayTogetherEventReceivedDelegate;
 
 /**
@@ -593,13 +595,10 @@ public:
 	 */
 	DEFINE_ONLINE_DELEGATE_TWO_PARAM(OnOnlineEnvironmentChanged, EOnlineEnvironment::Type /*LastEnvironment*/, EOnlineEnvironment::Type /*Environment*/);
 
-	/**
-	* Delegate fired when the "Play Together" event is sent from the PS4 system
-	*
-	* @param UserIndex - User index of the player the event is for
-	* @param UserIdList - list of other users in the PS4 party to send invites to
-	*/
-	DEFINE_ONLINE_DELEGATE_TWO_PARAM(OnPlayTogetherEventReceived, int32, TArray<TSharedPtr<const FUniqueNetId>>);
+	PRAGMA_DISABLE_DEPRECATION_WARNINGS
+	UE_DEPRECATED(4.26, "PlayTogether will no longer be supported and should be removed." )
+	FDelegateHandle AddOnPlayTogetherEventReceivedDelegate_Handle(const FOnPlayTogetherEventReceivedDelegate& Delegate) { return FDelegateHandle(); }
+	PRAGMA_ENABLE_DEPRECATION_WARNINGS
 
 	/**
 	 * @return The name of the online service this platform uses
diff --git a/Engine/Plugins/Online/OnlineSubsystemUtils/Source/OnlineSubsystemUtils/Classes/OnlineSessionClient.h b/Engine/Plugins/Online/OnlineSubsystemUtils/Source/OnlineSubsystemUtils/Classes/OnlineSessionClient.h
--- a/Engine/Plugins/Online/OnlineSubsystemUtils/Source/OnlineSubsystemUtils/Classes/OnlineSessionClient.h
+++ b/Engine/Plugins/Online/OnlineSubsystemUtils/Source/OnlineSubsystemUtils/Classes/OnlineSessionClient.h
@@ -35,8 +35,6 @@ protected:
 	FOnJoinSessionCompleteDelegate OnJoinSessionCompleteDelegate;
 	/** Delegate for accepting session invites */
 	FOnSessionUserInviteAcceptedDelegate OnSessionUserInviteAcceptedDelegate;
-	/** Delegate for handling the play together system event */
-	FOnPlayTogetherEventReceivedDelegate OnPlayTogetherEventReceivedDelegate;
 
 	// Handles to the above delegates
 	FDelegateHandle OnSessionInviteAcceptedDelegateHandle;
@@ -44,7 +42,6 @@ protected:
 	FDelegateHandle OnDestroyForJoinSessionCompleteDelegateHandle;
 	FDelegateHandle OnDestroyForMainMenuCompleteDelegateHandle;
 	FDelegateHandle OnJoinSessionCompleteDelegateHandle;
-	FDelegateHandle OnPlayTogetherEventReceivedDelegateHandle;
 
 	/** Handle to outstanding start session call */
 	FDelegateHandle StartSessionCompleteHandle;
diff --git a/Engine/Plugins/Online/OnlineSubsystemUtils/Source/OnlineSubsystemUtils/Private/OnlineSessionClient.cpp b/Engine/Plugins/Online/OnlineSubsystemUtils/Source/OnlineSubsystemUtils/Private/OnlineSessionClient.cpp
--- a/Engine/Plugins/Online/OnlineSubsystemUtils/Source/OnlineSubsystemUtils/Private/OnlineSessionClient.cpp
+++ b/Engine/Plugins/Online/OnlineSubsystemUtils/Source/OnlineSubsystemUtils/Private/OnlineSessionClient.cpp
@@ -46,12 +46,6 @@ void UOnlineSessionClient::RegisterOnlineDelegates()
 	OnDestroyForJoinSessionCompleteDelegate = FOnDestroySessionCompleteDelegate::CreateUObject(this, &ThisClass::OnDestroyForJoinSessionComplete);
 	OnDestroyForMainMenuCompleteDelegate	= FOnDestroySessionCompleteDelegate::CreateUObject(this, &ThisClass::OnDestroyForMainMenuComplete);
 	OnSessionUserInviteAcceptedDelegate     = FOnSessionUserInviteAcceptedDelegate::CreateUObject(this, &ThisClass::OnSessionUserInviteAccepted);
-	OnPlayTogetherEventReceivedDelegate		= FOnPlayTogetherEventReceivedDelegate::CreateUObject(this, &ThisClass::OnPlayTogetherEventReceived);
-
-	if (IOnlineSubsystem* const OnlineSubsystem = IOnlineSubsystem::Get())
-	{
-		OnPlayTogetherEventReceivedDelegateHandle = OnlineSubsystem->AddOnPlayTogetherEventReceivedDelegate_Handle(OnPlayTogetherEventReceivedDelegate);
-	}
 
 	IOnlineSessionPtr SessionInt = GetSessionInt();
 	if (SessionInt.IsValid())
@@ -67,11 +61,6 @@ void UOnlineSessionClient::ClearOnlineDelegates()
 	{
 		SessionInt->ClearOnSessionUserInviteAcceptedDelegate_Handle(OnSessionUserInviteAcceptedDelegateHandle);
 	}
-
-	if (IOnlineSubsystem* const OnlineSubsystem = IOnlineSubsystem::Get())
-	{
-		OnlineSubsystem->ClearOnPlayTogetherEventReceivedDelegate_Handle(OnPlayTogetherEventReceivedDelegateHandle);
-	}
 }
 
 void UOnlineSessionClient::OnSessionUserInviteAccepted(bool bWasSuccessful, int32 ControllerId, TSharedPtr<const FUniqueNetId> UserId, const FOnlineSessionSearchResult& SearchResult)
diff --git a/Engine/Source/Runtime/Engine/Classes/GameFramework/OnlineSession.h b/Engine/Source/Runtime/Engine/Classes/GameFramework/OnlineSession.h
--- a/Engine/Source/Runtime/Engine/Classes/GameFramework/OnlineSession.h
+++ b/Engine/Source/Runtime/Engine/Classes/GameFramework/OnlineSession.h
@@ -40,6 +40,7 @@ public:
 	virtual void OnSessionUserInviteAccepted(const bool bWasSuccess, const int32 ControllerId, TSharedPtr< const FUniqueNetId > UserId, const FOnlineSessionSearchResult & InviteResult) {};
 
 	/** Called when the play together system event is received on PS4 */
+	UE_DEPRECATED(4.26, "PlayTogether will no longer be supported and should be removed.")
 	virtual void OnPlayTogetherEventReceived(int32 UserIndex, TArray<TSharedPtr<const FUniqueNetId>> UserIdList) {};
 };
 
