#online Fixing static code analysis issues in OnlineSubsystemUtils and InAppPurchase*Proxy2 interface for conditional operator usage.

diff --git a/Engine/Plugins/Online/OnlineSubsystemUtils/Source/OnlineSubsystemUtils/Private/InAppPurchaseCallbackProxy2.cpp b/Engine/Plugins/Online/OnlineSubsystemUtils/Source/OnlineSubsystemUtils/Private/InAppPurchaseCallbackProxy2.cpp
--- a/Engine/Plugins/Online/OnlineSubsystemUtils/Source/OnlineSubsystemUtils/Private/InAppPurchaseCallbackProxy2.cpp
+++ b/Engine/Plugins/Online/OnlineSubsystemUtils/Source/OnlineSubsystemUtils/Private/InAppPurchaseCallbackProxy2.cpp
@@ -22,8 +22,15 @@ UInAppPurchaseCallbackProxy2::UInAppPurchaseCallbackProxy2(const FObjectInitiali
 void UInAppPurchaseCallbackProxy2::Trigger(APlayerController* PlayerController, const FInAppPurchaseProductRequest2& ProductRequest)
 {
 	bFailedToEvenSubmit = true;
-	WorldPtr = (PlayerController != nullptr) ? PlayerController->GetWorld() : nullptr;
-	if (APlayerState* PlayerState = (PlayerController != nullptr) ? ToRawPtr(PlayerController->PlayerState) : nullptr)
+	WorldPtr = nullptr;
+	APlayerState* PlayerState = nullptr;
+	if (PlayerController != nullptr)
+	{
+		WorldPtr = PlayerController->GetWorld();
+		PlayerState = ToRawPtr(PlayerController->PlayerState);
+	}
+
+	if (PlayerState != nullptr)
 	{
 		if (IOnlineSubsystem* const OnlineSub = IOnlineSubsystem::IsLoaded() ? IOnlineSubsystem::Get() : nullptr)
 		{
@@ -64,8 +71,15 @@ void UInAppPurchaseCallbackProxy2::Trigger(APlayerController* PlayerController,
 
 void UInAppPurchaseCallbackProxy2::TriggerGetUnprocessedPurchases(APlayerController* PlayerController)
 {
-	WorldPtr = (PlayerController != nullptr) ? PlayerController->GetWorld() : nullptr;
-	if (APlayerState* PlayerState = (PlayerController != nullptr) ? ToRawPtr(PlayerController->PlayerState) : nullptr)
+	WorldPtr = nullptr;
+	APlayerState* PlayerState = nullptr;
+	if (PlayerController != nullptr)
+	{
+		WorldPtr = PlayerController->GetWorld();
+		PlayerState = ToRawPtr(PlayerController->PlayerState);
+	}
+
+	if (PlayerState != nullptr)
 	{
 		if (IOnlineSubsystem* const OnlineSub = IOnlineSubsystem::IsLoaded() ? IOnlineSubsystem::Get() : nullptr)
 		{
@@ -94,8 +108,15 @@ void UInAppPurchaseCallbackProxy2::TriggerGetUnprocessedPurchases(APlayerControl
 void UInAppPurchaseCallbackProxy2::TriggerGetOwnedPurchases(APlayerController* PlayerController)
 {
 	bFailedToEvenSubmit = true;
-	WorldPtr = (PlayerController != nullptr) ? PlayerController->GetWorld() : nullptr;
-	if (APlayerState* PlayerState = (PlayerController != nullptr) ? ToRawPtr(PlayerController->PlayerState) : nullptr)
+	WorldPtr = nullptr;
+	APlayerState* PlayerState = nullptr;
+	if (PlayerController != nullptr)
+	{
+		WorldPtr = PlayerController->GetWorld();
+		PlayerState = ToRawPtr(PlayerController->PlayerState);
+	}
+
+	if (PlayerState != nullptr)
 	{
 		if (IOnlineSubsystem* const OnlineSub = IOnlineSubsystem::IsLoaded() ? IOnlineSubsystem::Get() : nullptr)
 		{
diff --git a/Engine/Plugins/Online/OnlineSubsystemUtils/Source/OnlineSubsystemUtils/Private/InAppPurchaseQueryCallbackProxy2.cpp b/Engine/Plugins/Online/OnlineSubsystemUtils/Source/OnlineSubsystemUtils/Private/InAppPurchaseQueryCallbackProxy2.cpp
--- a/Engine/Plugins/Online/OnlineSubsystemUtils/Source/OnlineSubsystemUtils/Private/InAppPurchaseQueryCallbackProxy2.cpp
+++ b/Engine/Plugins/Online/OnlineSubsystemUtils/Source/OnlineSubsystemUtils/Private/InAppPurchaseQueryCallbackProxy2.cpp
@@ -19,8 +19,15 @@ void UInAppPurchaseQueryCallbackProxy2::TriggerQuery(APlayerController* PlayerCo
 {
 	bool bFailedToEvenSubmit = true;
 
-	WorldPtr = (PlayerController != nullptr) ? PlayerController->GetWorld() : nullptr;
-	if (APlayerState* PlayerState = (PlayerController != NULL) ? ToRawPtr(PlayerController->PlayerState) : nullptr)
+	WorldPtr = nullptr;
+	APlayerState* PlayerState = nullptr;
+	if (PlayerController != nullptr)
+	{
+		WorldPtr = PlayerController->GetWorld();
+		PlayerState = ToRawPtr(PlayerController->PlayerState);
+	}
+
+	if (PlayerState != nullptr)
 	{
 		if (IOnlineSubsystem* const OnlineSub = IOnlineSubsystem::IsLoaded() ? IOnlineSubsystem::Get() : nullptr)
 		{
diff --git a/Engine/Plugins/Online/OnlineSubsystemUtils/Source/OnlineSubsystemUtils/Private/InAppPurchaseRestoreCallbackProxy2.cpp b/Engine/Plugins/Online/OnlineSubsystemUtils/Source/OnlineSubsystemUtils/Private/InAppPurchaseRestoreCallbackProxy2.cpp
--- a/Engine/Plugins/Online/OnlineSubsystemUtils/Source/OnlineSubsystemUtils/Private/InAppPurchaseRestoreCallbackProxy2.cpp
+++ b/Engine/Plugins/Online/OnlineSubsystemUtils/Source/OnlineSubsystemUtils/Private/InAppPurchaseRestoreCallbackProxy2.cpp
@@ -22,8 +22,15 @@ void UInAppPurchaseRestoreCallbackProxy2::Trigger(const TArray<FInAppPurchasePro
 	bWasSuccessful = false;
 	bool bFailedToEvenSubmit = true;
 
-	WorldPtr = (PlayerController != nullptr) ? PlayerController->GetWorld() : nullptr;
-	if (APlayerState* PlayerState = (PlayerController != nullptr) ? ToRawPtr(PlayerController->PlayerState) : nullptr)
+	WorldPtr = nullptr;
+	APlayerState* PlayerState = nullptr;
+	if (PlayerController != nullptr)
+	{
+		WorldPtr = PlayerController->GetWorld();
+		PlayerState = ToRawPtr(PlayerController->PlayerState);
+	}
+
+	if (PlayerState != nullptr)
 	{
 		if (IOnlineSubsystem* const OnlineSub = IOnlineSubsystem::IsLoaded() ? IOnlineSubsystem::Get() : nullptr)
 		{
diff --git a/Engine/Plugins/Online/OnlineSubsystemUtils/Source/OnlineSubsystemUtils/Private/OnlineSubsystemUtils.cpp b/Engine/Plugins/Online/OnlineSubsystemUtils/Source/OnlineSubsystemUtils/Private/OnlineSubsystemUtils.cpp
--- a/Engine/Plugins/Online/OnlineSubsystemUtils/Source/OnlineSubsystemUtils/Private/OnlineSubsystemUtils.cpp
+++ b/Engine/Plugins/Online/OnlineSubsystemUtils/Source/OnlineSubsystemUtils/Private/OnlineSubsystemUtils.cpp
@@ -809,7 +809,13 @@ void FOnlineSubsystemBPCallHelper::QueryIDFromPlayerController(APlayerController
 {
 	UserID.Reset();
 
-	if (APlayerState* PlayerState = (PlayerController != NULL) ? ToRawPtr(PlayerController->PlayerState) : NULL)
+	APlayerState* PlayerState = nullptr;
+	if (PlayerController != nullptr)
+	{
+		PlayerState = ToRawPtr(PlayerController->PlayerState);
+	}
+
+	if (PlayerState != nullptr)
 	{
 		UserID = PlayerState->GetUniqueId().GetUniqueNetId();
 		if (!UserID.IsValid())
