#online Fixing static code analysis issues in InAppPurchase*Proxy2 interface for a missing nullptr check on PlayerController, this can never happen, but the static analyser does not see the state handling.

diff --git a/Engine/Plugins/Online/OnlineSubsystemUtils/Source/OnlineSubsystemUtils/Private/InAppPurchaseCallbackProxy2.cpp b/Engine/Plugins/Online/OnlineSubsystemUtils/Source/OnlineSubsystemUtils/Private/InAppPurchaseCallbackProxy2.cpp
--- a/Engine/Plugins/Online/OnlineSubsystemUtils/Source/OnlineSubsystemUtils/Private/InAppPurchaseCallbackProxy2.cpp
+++ b/Engine/Plugins/Online/OnlineSubsystemUtils/Source/OnlineSubsystemUtils/Private/InAppPurchaseCallbackProxy2.cpp
@@ -45,6 +45,7 @@ void UInAppPurchaseCallbackProxy2::Trigger(APlayerController* PlayerController,
 				// Set-up, and trigger the transaction through the store interface
 				FPurchaseCheckoutRequest CheckoutRequest = FPurchaseCheckoutRequest();
 				CheckoutRequest.AddPurchaseOffer("", ProductRequest.ProductIdentifier, 1, ProductRequest.bIsConsumable);
+				check(PlayerController);
 				PurchasingPlayer = (*PlayerController->GetLocalPlayer()->GetUniqueNetIdFromCachedControllerId()).AsShared();
 				PurchaseInterface->Checkout(*PurchasingPlayer, CheckoutRequest, InAppPurchaseCompleteDelegate);
 			}
@@ -86,6 +87,7 @@ void UInAppPurchaseCallbackProxy2::TriggerGetUnprocessedPurchases(APlayerControl
 			PurchaseInterface = OnlineSub->GetPurchaseInterface();
 			if (PurchaseInterface.IsValid())
 			{
+				check(PlayerController);
 				PurchasingPlayer = (*PlayerController->GetLocalPlayer()->GetUniqueNetIdFromCachedControllerId()).AsShared();
 				OnCheckoutComplete(FOnlineError(true), MakeShared<FPurchaseReceipt>());
 			}
@@ -123,6 +125,7 @@ void UInAppPurchaseCallbackProxy2::TriggerGetOwnedPurchases(APlayerController* P
 			PurchaseInterface = OnlineSub->GetPurchaseInterface();
 			if (PurchaseInterface.IsValid())
 			{
+				check(PlayerController);
 				PurchasingPlayer = (*PlayerController->GetLocalPlayer()->GetUniqueNetIdFromCachedControllerId()).AsShared();
 				PurchaseInterface->QueryReceipts(*PurchasingPlayer, false, FOnQueryReceiptsComplete::CreateUObject(this, &UInAppPurchaseCallbackProxy2::OnQueryReceiptsComplete));
 			}
diff --git a/Engine/Plugins/Online/OnlineSubsystemUtils/Source/OnlineSubsystemUtils/Private/InAppPurchaseQueryCallbackProxy2.cpp b/Engine/Plugins/Online/OnlineSubsystemUtils/Source/OnlineSubsystemUtils/Private/InAppPurchaseQueryCallbackProxy2.cpp
--- a/Engine/Plugins/Online/OnlineSubsystemUtils/Source/OnlineSubsystemUtils/Private/InAppPurchaseQueryCallbackProxy2.cpp
+++ b/Engine/Plugins/Online/OnlineSubsystemUtils/Source/OnlineSubsystemUtils/Private/InAppPurchaseQueryCallbackProxy2.cpp
@@ -35,6 +35,7 @@ void UInAppPurchaseQueryCallbackProxy2::TriggerQuery(APlayerController* PlayerCo
 			if (StoreInterface.IsValid())
 			{
 				FFrame::KismetExecutionMessage(TEXT("UInAppPurchaseQueryCallbackProxy2::TriggerQuery - GetUniqueNetIdFromCachedControllerId"), ELogVerbosity::Warning);
+				check(PlayerController);
 				FUniqueNetIdRepl QueryingPlayer = PlayerController->GetLocalPlayer()->GetUniqueNetIdFromCachedControllerId();
 
 				//if (QueryingPlayer.IsValid())
diff --git a/Engine/Plugins/Online/OnlineSubsystemUtils/Source/OnlineSubsystemUtils/Private/InAppPurchaseRestoreCallbackProxy2.cpp b/Engine/Plugins/Online/OnlineSubsystemUtils/Source/OnlineSubsystemUtils/Private/InAppPurchaseRestoreCallbackProxy2.cpp
--- a/Engine/Plugins/Online/OnlineSubsystemUtils/Source/OnlineSubsystemUtils/Private/InAppPurchaseRestoreCallbackProxy2.cpp
+++ b/Engine/Plugins/Online/OnlineSubsystemUtils/Source/OnlineSubsystemUtils/Private/InAppPurchaseRestoreCallbackProxy2.cpp
@@ -39,6 +39,7 @@ void UInAppPurchaseRestoreCallbackProxy2::Trigger(const TArray<FInAppPurchasePro
 			{
 				bFailedToEvenSubmit = false;
 
+				check(PlayerController);
 				PurchasingPlayer = (*PlayerController->GetLocalPlayer()->GetUniqueNetIdFromCachedControllerId()).AsShared();
 				PurchaseInterface->QueryReceipts(*PurchasingPlayer, true, FOnQueryReceiptsComplete::CreateUObject(this, &UInAppPurchaseRestoreCallbackProxy2::OnQueryReceiptsComplete));
 			}
