OSSTests First draft. Bare minimum required for someone to be able to come on and start writing integration tests for OSS code.

diff --git a/Engine/Plugins/Online/OnlineSubsystemUtils/Source/OnlineSubsystemUtils/OnlineSubsystemUtils.Build.cs b/Engine/Plugins/Online/OnlineSubsystemUtils/Source/OnlineSubsystemUtils/OnlineSubsystemUtils.Build.cs
--- a/Engine/Plugins/Online/OnlineSubsystemUtils/Source/OnlineSubsystemUtils/OnlineSubsystemUtils.Build.cs
+++ b/Engine/Plugins/Online/OnlineSubsystemUtils/Source/OnlineSubsystemUtils/OnlineSubsystemUtils.Build.cs
@@ -24,15 +24,28 @@ public class OnlineSubsystemUtils : ModuleRules
             PrivateIncludePaths.Add(RuntimePath + "Windows/XAudio2/Private");
         }
 
+        if (Target.bCompileAgainstEngine)
+        {
+            PublicDependencyModuleNames.AddRange(
+                new string[] {
+                    "Engine"
+				}
+            );
+
+			PrivateDependencyModuleNames.AddRange(
+				new string[] {
+					"Voice",
+					"AudioMixer"
+				}
+			);
+        }
 
-		PrivateDependencyModuleNames.AddRange(
+        PrivateDependencyModuleNames.AddRange(
 			new string[] {
 				"ImageCore",
 				"Sockets",
-				"Voice",
 				"PacketHandler",
 				"Json",
-				"AudioMixer",
 				"SignalProcessing",
 				"AudioMixerCore",
 				"DeveloperSettings"
@@ -44,8 +57,7 @@ public class OnlineSubsystemUtils : ModuleRules
 			{
 				"Core",
 				"CoreUObject",
-				"NetCore",
-				"Engine"
+				"NetCore"
 			}
 		);
 
