Improvements to Test Stats Interface

diff --git a/Engine/Plugins/Online/OnlineSubsystem/Source/Private/OnlineKeyValuePair.cpp b/Engine/Plugins/Online/OnlineSubsystem/Source/Private/OnlineKeyValuePair.cpp
--- a/Engine/Plugins/Online/OnlineSubsystem/Source/Private/OnlineKeyValuePair.cpp
+++ b/Engine/Plugins/Online/OnlineSubsystem/Source/Private/OnlineKeyValuePair.cpp
@@ -497,6 +497,64 @@ void FVariantData::GetValue(TArray<TSharedPtr<FJsonValue>>& OutData) const
 	}
 }
 
+/**
+ * Increments the numeric value by 1
+ */
+void FVariantData::Increment()
+{
+	checkSlow(IsNumeric());
+	switch (Type)
+	{
+	case EOnlineKeyValuePairDataType::Int32:
+		Increment<int32, EOnlineKeyValuePairDataType::Int32>(1);
+		break;
+	case EOnlineKeyValuePairDataType::UInt32:
+		Increment<uint32, EOnlineKeyValuePairDataType::UInt32>(1);
+		break;
+	case EOnlineKeyValuePairDataType::Int64:
+		Increment<int64, EOnlineKeyValuePairDataType::Int64>(1);
+		break;
+	case EOnlineKeyValuePairDataType::UInt64:
+		Increment<uint64, EOnlineKeyValuePairDataType::UInt64>(1);
+		break;
+	case EOnlineKeyValuePairDataType::Double:
+		Increment<double, EOnlineKeyValuePairDataType::Double>(1.0);
+		break;
+	case EOnlineKeyValuePairDataType::Float:
+		Increment<float, EOnlineKeyValuePairDataType::Float>(1.0);
+		break;
+	}
+}
+
+/**
+ * Decrements the numeric value by 1
+ */
+void FVariantData::Decrement()
+{
+	checkSlow(IsNumeric());
+	switch (Type)
+	{
+	case EOnlineKeyValuePairDataType::Int32:
+		Decrement<int32, EOnlineKeyValuePairDataType::Int32>(1);
+		break;
+	case EOnlineKeyValuePairDataType::UInt32:
+		Decrement<uint32, EOnlineKeyValuePairDataType::UInt32>(1);
+		break;
+	case EOnlineKeyValuePairDataType::Int64:
+		Decrement<int64, EOnlineKeyValuePairDataType::Int64>(1);
+		break;
+	case EOnlineKeyValuePairDataType::UInt64:
+		Decrement<uint64, EOnlineKeyValuePairDataType::UInt64>(1);
+		break;
+	case EOnlineKeyValuePairDataType::Double:
+		Decrement<double, EOnlineKeyValuePairDataType::Double>(1.0);
+		break;
+	case EOnlineKeyValuePairDataType::Float:
+		Decrement<float, EOnlineKeyValuePairDataType::Float>(1.0);
+		break;
+	}
+}
+
 /**
  * Cleans up the existing data and sets the type to EOnlineKeyValuePairDataType::Empty
  */
diff --git a/Engine/Plugins/Online/OnlineSubsystem/Source/Public/OnlineKeyValuePair.h b/Engine/Plugins/Online/OnlineSubsystem/Source/Public/OnlineKeyValuePair.h
--- a/Engine/Plugins/Online/OnlineSubsystem/Source/Public/OnlineKeyValuePair.h
+++ b/Engine/Plugins/Online/OnlineSubsystem/Source/Public/OnlineKeyValuePair.h
@@ -435,6 +435,11 @@ public:
 		}
 	}
 
+	/**
+	 * Increments the numeric value by 1
+	 */
+	void Increment();
+
 	/**
 	 * Decrements the value by the specified amount
 	 *
@@ -450,6 +455,11 @@ public:
 		}
 	}
 
+	/**
+	 * Decrements the numeric value by 1
+	 */
+	void Decrement();
+
 	/**
 	 * Converts the data into a string representation
 	 */
diff --git a/Engine/Plugins/Online/OnlineSubsystemUtils/Source/OnlineSubsystemUtils/Private/OnlineSubsystemUtils.cpp b/Engine/Plugins/Online/OnlineSubsystemUtils/Source/OnlineSubsystemUtils/Private/OnlineSubsystemUtils.cpp
--- a/Engine/Plugins/Online/OnlineSubsystemUtils/Source/OnlineSubsystemUtils/Private/OnlineSubsystemUtils.cpp
+++ b/Engine/Plugins/Online/OnlineSubsystemUtils/Source/OnlineSubsystemUtils/Private/OnlineSubsystemUtils.cpp
@@ -631,7 +631,7 @@ static bool OnlineExec( UWorld* InWorld, const TCHAR* Cmd, FOutputDevice& Ar )
 					}
 					else if (FParse::Command(&Cmd, TEXT("STATS")))
 					{
-						(new FTestStatsInterface(SubName))->Test(InWorld);
+						(new FTestStatsInterface(SubName))->Test(InWorld, Cmd);
 						bWasHandled = true;
 					}
 					else if (FParse::Command(&Cmd, TEXT("PRESENCE")))
diff --git a/Engine/Plugins/Online/OnlineSubsystemUtils/Source/OnlineSubsystemUtils/Private/Tests/TestStatsInterface.cpp b/Engine/Plugins/Online/OnlineSubsystemUtils/Source/OnlineSubsystemUtils/Private/Tests/TestStatsInterface.cpp
--- a/Engine/Plugins/Online/OnlineSubsystemUtils/Source/OnlineSubsystemUtils/Private/Tests/TestStatsInterface.cpp
+++ b/Engine/Plugins/Online/OnlineSubsystemUtils/Source/OnlineSubsystemUtils/Private/Tests/TestStatsInterface.cpp
@@ -10,10 +10,9 @@ FTestStatsInterface::FTestStatsInterface(const FString& InSubsystem) :
 	Subsystem(InSubsystem),
 	bOverallSuccess(true),
 	Stats(NULL),
-	TestPhase(0),
-	LastTestPhase(-1)
+	CurrentTestPhase(EStatsTestPhase::ReadStatsForOneUser),
+	LastTestPhase(EStatsTestPhase::Invalid)
 {
-
 }
 
 FTestStatsInterface::~FTestStatsInterface()
@@ -21,12 +20,12 @@ FTestStatsInterface::~FTestStatsInterface()
 	Stats = NULL;
 }
 
-void FTestStatsInterface::Test(UWorld* InWorld)
+void FTestStatsInterface::Test(UWorld* InWorld, const TCHAR* Cmd)
 {
 	OnlineSub = Online::GetSubsystem(InWorld, FName(*Subsystem));
 	if (!OnlineSub)
 	{
-		UE_LOG_ONLINE_STATS(Warning, TEXT("Failed to get online subsystem for %s"), *Subsystem);
+		UE_LOG_ONLINE_STATS(Warning, TEXT("[FTestStatsInterface::Test] Failed to get online subsystem %s"), *Subsystem);
 
 		bOverallSuccess = false;
 		return;
@@ -35,123 +34,223 @@ void FTestStatsInterface::Test(UWorld* InWorld)
 	if (OnlineSub->GetIdentityInterface().IsValid())
 	{
 		UserId = OnlineSub->GetIdentityInterface()->GetUniquePlayerId(0);
+		
+		if (!UserId.IsValid())
+		{
+			UE_LOG_ONLINE_STATS(Warning, TEXT("[FTestStatsInterface::Test] Couldn't find any valid user in identity interface"));
+
+			bOverallSuccess = false;
+			return;
+		}
 	}
 
 	// Cache interfaces
 	Stats = OnlineSub->GetStatsInterface();
 	if (!Stats.IsValid())
 	{
-		UE_LOG_ONLINE_STATS(Warning, TEXT("Failed to get online Stats interface for %s"), *Subsystem);
+		UE_LOG_ONLINE_STATS(Warning, TEXT("[FTestStatsInterface::Test] Failed to get online Stats interface for online subsystem %s"), *Subsystem);
 
 		bOverallSuccess = false;
 		return;
 	}
-}
 
-void FTestStatsInterface::WriteStats()
-{
-	TArray<FOnlineStatsUserUpdatedStats> Writes;
+	FString StatToRead = FParse::Token(Cmd, false);
+	while (!StatToRead.IsEmpty())
+	{
+		StatsToRead.Add(StatToRead);
+		StatToRead = FParse::Token(Cmd, false);
+	}
+
+	if (StatsToRead.Num() == 0)
+	{
+		UE_LOG_ONLINE_STATS(Warning, TEXT("[FTestStatsInterface::Test] No stat names were specified. Please add the stat names you want to use for the test."), *Subsystem);
 
-	int64 NewScore = 999;
-	FVariantData NewScoreData;
-	NewScoreData.SetValue(NewScore);
+		bOverallSuccess = false;
+		return;
+	}
+}
 
-	int64 NewFrags = 53;
-	FVariantData NewFragsData;
-	NewFragsData.SetValue(NewFrags);
+void FTestStatsInterface::WriteStats(bool bIncrementStats)
+{
+	UE_LOG_ONLINE_STATS(Log, TEXT("[FTestStatsInterface::WriteStats]"));
 
-	int64 NewDeaths = 24;
-	FVariantData NewDeathsData;
-	NewDeathsData.SetValue(NewDeaths);
+	TArray<FOnlineStatsUserUpdatedStats> Writes;
+	FOnlineStatsUserUpdatedStats& UpdatedValues = Writes.Emplace_GetRef(UserId.ToSharedRef());
 
-	int64 NewMatchesPlayed = 3;
-	FVariantData NewMatchesPlayedData;
-	NewMatchesPlayedData.SetValue(NewMatchesPlayed);
+	TSharedPtr<const FOnlineStatsUserStats> ReadUserStats = Stats->GetStats(UserId.ToSharedRef());
+	if (ReadUserStats.IsValid())
+	{
+		for (const FString& StatName : StatsToRead)
+		{
+			if (const FOnlineStatValue* StatRead = ReadUserStats->Stats.Find(StatName))
+			{
+				if (StatRead->IsNumeric())
+				{
+					FOnlineStatValue NewData(*StatRead);
 
-	UE_LOG_ONLINE_STATS(Log, TEXT("FTestStatsInterface::WriteStats()"));
-	UE_LOG_ONLINE_STATS(Log, TEXT("	- ShooterAllTimeMatchResultsScore is being set to %ll"), NewScore);
-	UE_LOG_ONLINE_STATS(Log, TEXT("	- ShooterAllTimeMatchResultsFrags is being incremented by %ll"), NewFrags);
-	UE_LOG_ONLINE_STATS(Log, TEXT("	- ShooterAllTimeMatchResultsDeaths is being set to the largest of its existing value and %ll"), NewDeaths);
-	UE_LOG_ONLINE_STATS(Log, TEXT("	- ShooterAllTimeMatchResultsMatchesPlayed is being set to the smallest of its existng value and %ll"), NewMatchesPlayed);
+					bIncrementStats ? NewData.Increment() : NewData.Decrement();
 
-	FOnlineStatsUserUpdatedStats& Write1 = Writes.Emplace_GetRef(UserId.ToSharedRef());
-	Write1.Stats.Add(TEXT("ShooterAllTimeMatchResultsScore"), FOnlineStatUpdate(NewScoreData, FOnlineStatUpdate::EOnlineStatModificationType::Set));
-	Write1.Stats.Add(TEXT("ShooterAllTimeMatchResultsFrags"), FOnlineStatUpdate(NewFragsData, FOnlineStatUpdate::EOnlineStatModificationType::Sum));
-	Write1.Stats.Add(TEXT("ShooterAllTimeMatchResultsDeaths"), FOnlineStatUpdate(NewDeathsData, FOnlineStatUpdate::EOnlineStatModificationType::Largest));
-	Write1.Stats.Add(TEXT("ShooterAllTimeMatchResultsMatchesPlayed"), FOnlineStatUpdate(NewMatchesPlayedData, FOnlineStatUpdate::EOnlineStatModificationType::Smallest));
+					UE_LOG_ONLINE_STATS(Log, TEXT("Stat %s is being %s"), *StatName, bIncrementStats ? TEXT("incremented") : TEXT("decremented"));
+					UpdatedValues.Stats.Add(*StatName, FOnlineStatUpdate(NewData, FOnlineStatUpdate::EOnlineStatModificationType::Set));
+				}
+				else
+				{
+					UE_LOG_ONLINE_STATS(Log, TEXT("Stat %s doesn't have a numeric value, skipped"), *StatName);
+				}
+			}
+		}
+	}
+	else
+	{
+		UE_LOG_ONLINE_STATS(Log, TEXT("Stats not found for user %s"), *UserId->ToString());
+	}
 
 	// Write it to the buffers
 	Stats->UpdateStats(UserId.ToSharedRef(), Writes, FOnlineStatsUpdateStatsComplete::CreateLambda([this](const FOnlineError& Error) {
-		UE_LOG_ONLINE_STATS(Log, TEXT("Write test finish"));
-		TestPhase++;
+		UE_LOG_ONLINE_STATS(Log, TEXT("[FTestStatsInterface::WriteStats] Write test finished"));
+
+		CurrentTestPhase++;
 
 		if(!Error.bSucceeded)
 		{
-			UE_LOG_ONLINE_STATS(Error, TEXT("WriteStats test failed: %s"), *Error.GetErrorMessage().ToString());
+			UE_LOG_ONLINE_STATS(Error, TEXT("[FTestStatsInterface::WriteStats] WriteStats test failed with error: %s"), *Error.GetErrorMessage().ToString());
 			bOverallSuccess = false;
 		}
 	}));
 }
 
-void FTestStatsInterface::ReadStats()
+void FTestStatsInterface::ReadStats(bool bIncludeStatsToRead)
 {
-	Stats->QueryStats(UserId.ToSharedRef(), UserId.ToSharedRef(), FOnlineStatsQueryUserStatsComplete::CreateLambda([this](const FOnlineError& Error,  const TSharedPtr<const FOnlineStatsUserStats>& QueriedStats)
+	if (bIncludeStatsToRead)
+	{
+		Stats->QueryStats(UserId.ToSharedRef(), { UserId.ToSharedRef() }, StatsToRead, FOnlineStatsQueryUsersStatsComplete::CreateRaw(this, &FTestStatsInterface::OnQueryUsersStatsComplete));
+	}
+	else
 	{
-		UE_LOG_ONLINE_STATS(Log, TEXT("Read test finish with %d queried stats"), QueriedStats->Stats.Num());
+		Stats->QueryStats(UserId.ToSharedRef(), UserId.ToSharedRef(), FOnlineStatsQueryUserStatsComplete::CreateRaw(this, &FTestStatsInterface::OnQueryUserStatsComplete));
+	}
+}
+
+void FTestStatsInterface::OnQueryUserStatsComplete(const FOnlineError& ResultState, const TSharedPtr<const FOnlineStatsUserStats>& QueriedStats)
+{
+	UE_LOG_ONLINE_STATS(Log, TEXT("[FTestStatsInterface::OnQueryUserStatsComplete] Read test finish with %d queried stats"), QueriedStats->Stats.Num());
 
+	if (!ResultState.bSucceeded)
+	{
+		UE_LOG_ONLINE_STATS(Warning, TEXT("[FTestStatsInterface::OnQueryUserStatsComplete] ReadStats first attempt failed with error: %s"), *ResultState.GetErrorMessage().ToString());
 
-		if (!Error.bSucceeded)
+		if (ResultState.GetErrorResult() != EOnlineErrorResult::NotImplemented)
 		{
-			UE_LOG_ONLINE_STATS(Error, TEXT("ReadStats test failed: %s"), *Error.GetErrorMessage().ToString());
 			bOverallSuccess = false;
 		}
-		else
+	}
+	else
+	{
+		for (const TPair<FString, FOnlineStatValue>& StatReads : QueriedStats->Stats)
+		{
+			UE_LOG_ONLINE_STATS(Log, TEXT("[FTestStatsInterface::OnQueryUserStatsComplete] Queried stat with name %s has value %s"), *StatReads.Key, *StatReads.Value.ToString());
+		}
+
+		const TSharedPtr<const FOnlineStatsUserStats> ReadUserStats = Stats->GetStats(UserId.ToSharedRef());
+		if (ReadUserStats.IsValid())
+		{
+			for (const TPair<FString, FOnlineStatValue>& StatReads : ReadUserStats->Stats)
+			{
+				UE_LOG_ONLINE_STATS(Log, TEXT("[FTestStatsInterface::OnQueryUserStatsComplete] Cached stat with name %s has value %s"), *StatReads.Key, *StatReads.Value.ToString());
+			}
+		}
+	}
+
+	CurrentTestPhase++;
+}
+
+void FTestStatsInterface::OnQueryUsersStatsComplete(const FOnlineError& ResultState, const TArray<TSharedRef<const FOnlineStatsUserStats>>& UsersStatsResult)
+{
+	UE_LOG_ONLINE_STATS(Log, TEXT("[FTestStatsInterface::OnQueryUsersStatsComplete] Read test finish for %d users with queried stats"), UsersStatsResult.Num());
+
+	if (!ResultState.bSucceeded)
+	{
+		UE_LOG_ONLINE_STATS(Warning, TEXT("[FTestStatsInterface::OnQueryUsersStatsComplete] ReadStats first attempt failed with error: %s"), *ResultState.GetErrorMessage().ToString());
+
+		if (ResultState.GetErrorResult() != EOnlineErrorResult::NotImplemented)
+		{
+			bOverallSuccess = false;
+		}
+	}
+	else
+	{
+		for (const TSharedRef<const FOnlineStatsUserStats>& QueriedStats : UsersStatsResult)
 		{
 			for (const TPair<FString, FOnlineStatValue>& StatReads : QueriedStats->Stats)
 			{
-				UE_LOG_ONLINE_STATS(Log, TEXT("Stat name %s has value %s"), *StatReads.Key, *StatReads.Value.ToString());
+				UE_LOG_ONLINE_STATS(Log, TEXT("[FTestStatsInterface::OnQueryUsersStatsComplete] Queried stat with name %s has value %s"), *StatReads.Key, *StatReads.Value.ToString());
 			}
 
-			TSharedPtr<const FOnlineStatsUserStats> ReadUserStats = Stats->GetStats(UserId.ToSharedRef());
-			for (const TPair<FString, FOnlineStatValue>& StatReads : ReadUserStats->Stats)
+			// We would keep track of all queried users here, but we only used our local user
+			const TSharedPtr<const FOnlineStatsUserStats> ReadUserStats = Stats->GetStats(UserId.ToSharedRef());
+			if (ReadUserStats.IsValid())
 			{
-				UE_LOG_ONLINE_STATS(Log, TEXT("Read Stat name %s has value %s"), *StatReads.Key, *StatReads.Value.ToString());
+				for (const TPair<FString, FOnlineStatValue>& StatReads : ReadUserStats->Stats)
+				{
+					UE_LOG_ONLINE_STATS(Log, TEXT("[FTestStatsInterface::OnQueryUsersStatsComplete] Cached stat with name %s has value %s"), *StatReads.Key, *StatReads.Value.ToString());
+				}
 			}
-			TestPhase++;
-		}
-	}));
+		}	
+	}
+
+	CurrentTestPhase++;
 }
 
 bool FTestStatsInterface::Tick(float DeltaTime)
 {
 	QUICK_SCOPE_CYCLE_COUNTER(STAT_FTestStatsInterface_Tick);
 
-	if (TestPhase != LastTestPhase)
+	if (CurrentTestPhase != LastTestPhase)
 	{
 		if (!bOverallSuccess)
 		{
-			UE_LOG_ONLINE_STATS(Log, TEXT("Testing failed in phase %d"), LastTestPhase);
-			TestPhase = 3;
+			UE_LOG_ONLINE_STATS(Log, TEXT("[FTestStatsInterface::Tick] Testing failed in phase %d"), LastTestPhase);
+			CurrentTestPhase = EStatsTestPhase::End;
 		}
 
-		LastTestPhase = TestPhase;
+		LastTestPhase = CurrentTestPhase;
 
-		switch (TestPhase)
+		switch (CurrentTestPhase)
 		{
-		case 0:
-			UE_LOG_ONLINE_STATS(Log, TEXT("// Beginning ReadStats (reading self, pre-write)"));
-			ReadStats();
+		case EStatsTestPhase::ReadStatsForOneUser:
+			UE_LOG_ONLINE_STATS(Log, TEXT("[FTestStatsInterface::Tick] Beginning ReadStats (reading self, pre-increment, no stat names)"));
+			ReadStats(false);
+			break;
+		case EStatsTestPhase::ReadStatsForManyUsers:
+			UE_LOG_ONLINE_STATS(Log, TEXT("[FTestStatsInterface::Tick] Beginning ReadStats (reading self, pre-increment, including stat names)"));
+			ReadStats(true);
+			break;
+		case EStatsTestPhase::WriteIncrementedStats:
+			UE_LOG_ONLINE_STATS(Log, TEXT("[FTestStatsInterface::Tick] Beginning Write (incrementing stats)"));
+			WriteStats(true);
+			break;
+		case EStatsTestPhase::ReadStatsForOneUserAfterIncrement:
+			UE_LOG_ONLINE_STATS(Log, TEXT("[FTestStatsInterface::Tick] Beginning ReadStats (reading self, post-increment, no stat names)"));
+			ReadStats(false);
+			break;
+		case EStatsTestPhase::ReadStatsForManyUsersaAfterIncrement:
+			UE_LOG_ONLINE_STATS(Log, TEXT("[FTestStatsInterface::Tick] Beginning ReadStats (reading self, post-increment, including stat names)"));
+			ReadStats(true);
+			break;
+		case EStatsTestPhase::WriteDecrementedStats:
+			UE_LOG_ONLINE_STATS(Log, TEXT("[FTestStatsInterface::Tick] Beginning Write (decrementing stats)"));
+			WriteStats(false);
 			break;
-		case 1:
-			UE_LOG_ONLINE_STATS(Log, TEXT("// Beginning Write (writing matches to 999)"));
-			WriteStats();
+		case EStatsTestPhase::ReadStatsForOneUserAfterDecrement:
+			UE_LOG_ONLINE_STATS(Log, TEXT("[FTestStatsInterface::Tick] Beginning ReadStats (reading self, post-decrement, no stat names)"));
+			ReadStats(false);
 			break;
-		case 2:
-			UE_LOG_ONLINE_STATS(Log, TEXT("// Beginning ReadStats (reading self, post-write)"));
-			ReadStats();
+		case EStatsTestPhase::ReadStatsForManyUsersaAfterDecrement:
+			UE_LOG_ONLINE_STATS(Log, TEXT("[FTestStatsInterface::Tick] Beginning ReadStats (reading self, post-decrement, including stat names)"));
+			ReadStats(true);
 			break;
-		case 3:
-			UE_LOG_ONLINE_STATS(Log, TEXT("TESTING COMPLETE Success:%s!"), bOverallSuccess ? TEXT("true") : TEXT("false"));
+		case EStatsTestPhase::End:
+			UE_LOG_ONLINE_STATS(Log, TEXT("[FTestStatsInterface::Tick] TESTING COMPLETE Success: %s!"), *LexToString(bOverallSuccess));
 			delete this;
 			return false;
 		}
diff --git a/Engine/Plugins/Online/OnlineSubsystemUtils/Source/OnlineSubsystemUtils/Private/Tests/TestStatsInterface.h b/Engine/Plugins/Online/OnlineSubsystemUtils/Source/OnlineSubsystemUtils/Private/Tests/TestStatsInterface.h
--- a/Engine/Plugins/Online/OnlineSubsystemUtils/Source/OnlineSubsystemUtils/Private/Tests/TestStatsInterface.h
+++ b/Engine/Plugins/Online/OnlineSubsystemUtils/Source/OnlineSubsystemUtils/Private/Tests/TestStatsInterface.h
@@ -24,16 +24,46 @@
 	/** Keep track of success across all functions and callbacks */
 	bool bOverallSuccess;
 
+	/** Names of stats to read, taken from the console command */
+	TArray<FString> StatsToRead;
+
 	/** Logged in UserId */
 	FUniqueNetIdPtr UserId;
 
 	/** Convenient access to the Stats interfaces */
 	IOnlineStatsPtr Stats;
 
+	enum class EStatsTestPhase
+	{
+		Invalid,
+		ReadStatsForOneUser,
+		ReadStatsForManyUsers,
+		WriteIncrementedStats,
+		ReadStatsForOneUserAfterIncrement,
+		ReadStatsForManyUsersaAfterIncrement,
+		WriteDecrementedStats,
+		ReadStatsForOneUserAfterDecrement,
+		ReadStatsForManyUsersaAfterDecrement,
+		End
+	};
+
+	friend EStatsTestPhase& operator++(EStatsTestPhase& TestPhase) {
+		const int TestPhaseInt = int(TestPhase) + 1;
+		const int EndPhaseInt = int(EStatsTestPhase::End);
+		TestPhase = EStatsTestPhase((TestPhaseInt > EndPhaseInt) ? EndPhaseInt : TestPhaseInt);
+		return TestPhase;
+	}
+
+	friend EStatsTestPhase operator++(EStatsTestPhase& TestPhase, int) {
+		const EStatsTestPhase Result = TestPhase;
+		++TestPhase;
+		return Result;
+	}
+
 	/** Current phase of testing */
-	int32 TestPhase;
+	EStatsTestPhase CurrentTestPhase;
 	/** Last phase of testing triggered */
-	int32 LastTestPhase;
+	EStatsTestPhase LastTestPhase;
 
 	/** Hidden on purpose */
 	FTestStatsInterface()
@@ -43,18 +73,23 @@
 
 	/**
 	 *	Write out some test data to a Stats
+	 *
+	 * @param bIncrementStats If this is true, stats written will be incremented, if it's not, they will be decremented
 	 */
-	void WriteStats();
+	void WriteStats(bool bIncrementStats);
 
 	/**
-	 *	Delegate called when a Stats has been successfully read
+	 *	Delegates called when a Stats has been successfully read, both for a single and multiple users
 	 */
-	void OnStatsReadComplete(bool bWasSuccessful);
+	void OnQueryUserStatsComplete(const FOnlineError& Error, const TSharedPtr<const FOnlineStatsUserStats>& QueriedStats);
+	void OnQueryUsersStatsComplete(const FOnlineError& ResultState, const TArray<TSharedRef<const FOnlineStatsUserStats>>& UsersStatsResult);
 
 	/**
-	 *	Read in some predefined data from a Stats
+	 * Read values for stats
+	 *
+	 * @param bIncludeStatsToRead Specify if the call should include stats defined in command line
 	 */
-	void ReadStats();
+	void ReadStats(bool bIncludeStatsToRead);
 
 	/** Utilities */
 	void PrintStats();
@@ -78,7 +113,7 @@
 	/**
 	 * Kicks off all of the testing process
 	 */
-	void Test(UWorld* InWorld);
+	void Test(UWorld* InWorld, const TCHAR* Cmd);
  };
 
 #endif //WITH_DEV_AUTOMATION_TESTS
