Removing legacy audio backends and unrealaudio

diff --git a/Engine/Plugins/Online/OnlineSubsystemUtils/Source/OnlineSubsystemUtils/OnlineSubsystemUtils.Build.cs b/Engine/Plugins/Online/OnlineSubsystemUtils/Source/OnlineSubsystemUtils/OnlineSubsystemUtils.Build.cs
--- a/Engine/Plugins/Online/OnlineSubsystemUtils/Source/OnlineSubsystemUtils/OnlineSubsystemUtils.Build.cs
+++ b/Engine/Plugins/Online/OnlineSubsystemUtils/Source/OnlineSubsystemUtils/OnlineSubsystemUtils.Build.cs
@@ -20,8 +20,6 @@ public class OnlineSubsystemUtils : ModuleRules
         if (bIsWindowsPlatformBuild)
         {
             AddEngineThirdPartyPrivateStaticDependencies(Target, "DX11Audio");
-            PrivateIncludePaths.Add(RuntimePath + "Windows/XAudio2/Public");
-            PrivateIncludePaths.Add(RuntimePath + "Windows/XAudio2/Private");
         }
 
         if (Target.bCompileAgainstEngine)
diff --git a/Engine/Plugins/Online/OnlineSubsystemUtils/Source/OnlineSubsystemUtils/Private/VoiceEngineImpl.cpp b/Engine/Plugins/Online/OnlineSubsystemUtils/Source/OnlineSubsystemUtils/Private/VoiceEngineImpl.cpp
--- a/Engine/Plugins/Online/OnlineSubsystemUtils/Source/OnlineSubsystemUtils/Private/VoiceEngineImpl.cpp
+++ b/Engine/Plugins/Online/OnlineSubsystemUtils/Source/OnlineSubsystemUtils/Private/VoiceEngineImpl.cpp
@@ -14,14 +14,6 @@
 /** Largest size allowed to carry over into next buffer */
 #define MAX_VOICE_REMAINDER_SIZE 4 * 1024
 
-#if PLATFORM_WINDOWS
-#include "XAudio2Support.h"
-namespace NotificationClient
-{
-	TSharedPtr<FMMNotificationClient> WindowsNotificationClient;
-}
-#endif 
-
 namespace VoiceEngineUtilities
 {
 	void DownmixBuffer(const float* InAudio, float* OutAudio, int32 NumFrames, int32 InNumChannels, int32 OutNumChannels)
@@ -187,10 +179,6 @@ FVoiceEngineImpl ::FVoiceEngineImpl()
 	, bPendingFinalCapture(false)
 	, bIsCapturing(false)
 	, SerializeHelper(nullptr)
-#if PLATFORM_WINDOWS
-	, bAudioDeviceChanged(false)
-	, bDeviceChangeListenerRegistered(false)
-#endif
 {
 }
 
@@ -205,9 +193,6 @@ FVoiceEngineImpl::FVoiceEngineImpl(IOnlineSubsystem* InSubsystem) :
 	bPendingFinalCapture(false),
 	bIsCapturing(false),
 	SerializeHelper(nullptr)
-#if PLATFORM_WINDOWS
-	, bAudioDeviceChanged(false)
-#endif
 {
 	FCoreUObjectDelegates::PostLoadMapWithWorld.AddRaw(this, &FVoiceEngineImpl::OnPostLoadMap);
 
@@ -231,12 +216,6 @@ FVoiceEngineImpl::~FVoiceEngineImpl()
 
 	delete SerializeHelper;
 
-#if PLATFORM_WINDOWS
-	if (bDeviceChangeListenerRegistered)
-	{
-		UnregisterDeviceChangedListener();
-	}
-#endif
 }
 
 void FVoiceEngineImpl::VoiceCaptureUpdate() const
@@ -309,9 +288,6 @@ bool FVoiceEngineImpl::Init(int32 MaxLocalTalkers, int32 MaxRemoteTalkers)
 			bSuccess = VoiceEncoder.IsValid();
 			if (bSuccess)
 			{
-#if PLATFORM_WINDOWS
-				RegisterDeviceChangedListener();
-#endif
 				CompressedVoiceBuffer.Empty(UVOIPStatics::GetMaxCompressedVoiceDataSize());
 				DecompressedVoiceBuffer.Empty(UVOIPStatics::GetMaxUncompressedVoiceDataSizePerChannel());
 
@@ -688,13 +664,6 @@ void FVoiceEngineImpl::Tick(float DeltaTime)
 
 	// Push any buffered audio to any connected outputs.
 	AllRemoteTalkerAudio.ProcessAudio();
-
-#if PLATFORM_WINDOWS
-	if (bAudioDeviceChanged)
-	{
-		HandleDeviceChange();
-	}
-#endif
 }
 
 void FVoiceEngineImpl::GenerateVoiceData(USoundWaveProcedural* InProceduralWave, int32 SamplesRequired, const FUniqueNetId& TalkerId)
@@ -953,57 +922,6 @@ void FVoiceEngineImpl::CreateSerializeHelper()
 	}
 }
 
-#if PLATFORM_WINDOWS
-void FVoiceEngineImpl::RegisterDeviceChangedListener()
-{
-	if (!NotificationClient::WindowsNotificationClient.IsValid())
-	{
-		NotificationClient::WindowsNotificationClient = TSharedPtr<FMMNotificationClient>(new FMMNotificationClient);
-	}
-
-	NotificationClient::WindowsNotificationClient->RegisterDeviceChangedListener(this);
-
-	bDeviceChangeListenerRegistered = true;
-}
-
-void FVoiceEngineImpl::UnregisterDeviceChangedListener()
-{
-	if (NotificationClient::WindowsNotificationClient.IsValid())
-	{
-		NotificationClient::WindowsNotificationClient->UnRegisterDeviceDeviceChangedListener(this);
-	}
-
-	bDeviceChangeListenerRegistered = false;
-}
-
-void FVoiceEngineImpl::HandleDeviceChange()
-{
-	const double TimeSince = FPlatformTime::Seconds() - TimeDeviceChaned;
-	if (TimeSince >= DeviceChangeDelay)
-	{
-		if (bIsCapturing)
-		{
-			StopLocalVoiceProcessing(OwningUserIndex);
-			StartLocalVoiceProcessing(OwningUserIndex);
-		}
-
-		for (FRemoteTalkerData::TIterator It(RemoteTalkerBuffers); It; ++It)
-		{
-			FRemoteTalkerDataImpl& RemoteData = It.Value();
-			RemoteData.Reset();
-		}
-
-		bAudioDeviceChanged = false;
-	}
-}
-
-void FVoiceEngineImpl::OnDefaultDeviceChanged()
-{
-	bAudioDeviceChanged = true;
-	TimeDeviceChaned = FPlatformTime::Seconds();
-}
-#endif
-
 IOnlineSubsystem* FVoiceEngineImpl::GetOnlineSubSystem()
 {
 	if (UWorld* World = GetWorldForOnline(OnlineInstanceName))
diff --git a/Engine/Plugins/Online/OnlineSubsystemUtils/Source/OnlineSubsystemUtils/Public/VoiceEngineImpl.h b/Engine/Plugins/Online/OnlineSubsystemUtils/Source/OnlineSubsystemUtils/Public/VoiceEngineImpl.h
--- a/Engine/Plugins/Online/OnlineSubsystemUtils/Source/OnlineSubsystemUtils/Public/VoiceEngineImpl.h
+++ b/Engine/Plugins/Online/OnlineSubsystemUtils/Source/OnlineSubsystemUtils/Public/VoiceEngineImpl.h
@@ -219,13 +219,6 @@ class ONLINESUBSYSTEMUTILS_API FVoiceEngineImpl : public IVoiceEngine, public FS
 	 */
 	TArray<TUniquePtr<FVoiceEndpoint>> ExternalEndpoints;
 
-// Get Audio Device Changes on Windows
-#if PLATFORM_WINDOWS
-	FThreadSafeBool bAudioDeviceChanged;
-	double TimeDeviceChaned;
-	const double DeviceChangeDelay = 2.0f;
-#endif
-
 protected:
 	/**
 	 * Determines if the specified index is the owner or not
@@ -393,19 +386,7 @@ protected:
 	virtual int32							 GetMaxVoiceRemainderSize();
 	virtual void							 CreateSerializeHelper();
 
-	// Get Audio Device Changes on Windows
-#if PLATFORM_WINDOWS
-	virtual void RegisterDeviceChangedListener();
-	virtual void UnregisterDeviceChangedListener();
-	virtual void HandleDeviceChange();
-
-	//~ Begin IDeviceChangedListener
-	virtual void OnDefaultDeviceChanged() override;
-
-	bool bDeviceChangeListenerRegistered;
-#else
 	virtual void OnDefaultDeviceChanged() override {}
-#endif
 	virtual void OnDeviceRemoved(FString DeviceID) override {}
 	//~ End IDeviceChangedListener
 };
diff --git a/Engine/Plugins/Runtime/Oculus/OculusAudio/Source/OculusAudio/OculusAudio.Build.cs b/Engine/Plugins/Runtime/Oculus/OculusAudio/Source/OculusAudio/OculusAudio.Build.cs
--- a/Engine/Plugins/Runtime/Oculus/OculusAudio/Source/OculusAudio/OculusAudio.Build.cs
+++ b/Engine/Plugins/Runtime/Oculus/OculusAudio/Source/OculusAudio/OculusAudio.Build.cs
@@ -42,9 +42,6 @@ namespace UnrealBuildTool.Rules
 
 			if (Target.Platform == UnrealTargetPlatform.Win64)
 			{
-				PrivateIncludePathModuleNames.Add("XAudio2");
-				PrivateDependencyModuleNames.AddRange(new string[] { "XAudio2" });
-
 				// Automatically copy DLL to packaged builds
 				RuntimeDependencies.Add("$(EngineDir)/Binaries/ThirdParty/Oculus/Audio/Win64/ovraudio64.dll");
 
diff --git a/Engine/Plugins/Runtime/Oculus/OculusAudio/Source/OculusAudio/Private/OculusAudio.cpp b/Engine/Plugins/Runtime/Oculus/OculusAudio/Source/OculusAudio/Private/OculusAudio.cpp
--- a/Engine/Plugins/Runtime/Oculus/OculusAudio/Source/OculusAudio/Private/OculusAudio.cpp
+++ b/Engine/Plugins/Runtime/Oculus/OculusAudio/Source/OculusAudio/Private/OculusAudio.cpp
@@ -12,18 +12,7 @@ TAudioSpatializationPtr FOculusSpatializationPluginFactory::CreateNewSpatializat
 	check(Plugin != nullptr);
 	Plugin->RegisterAudioDevice(OwningDevice);
 
-#if PLATFORM_WINDOWS
-	if (OwningDevice->IsAudioMixerEnabled())
-#endif
-	{
-		return TAudioSpatializationPtr(new OculusAudioSpatializationAudioMixer());
-	}
-#if PLATFORM_WINDOWS
-	else
-	{
-		return TAudioSpatializationPtr(new OculusAudioLegacySpatialization());
-	}
-#endif
+	return TAudioSpatializationPtr(new OculusAudioSpatializationAudioMixer());
 }
 
 
diff --git a/Engine/Plugins/Runtime/Oculus/OculusAudio/Source/OculusAudio/Private/OculusAudio.h b/Engine/Plugins/Runtime/Oculus/OculusAudio/Source/OculusAudio/Private/OculusAudio.h
--- a/Engine/Plugins/Runtime/Oculus/OculusAudio/Source/OculusAudio/Private/OculusAudio.h
+++ b/Engine/Plugins/Runtime/Oculus/OculusAudio/Source/OculusAudio/Private/OculusAudio.h
@@ -4,9 +4,6 @@
 
 #include "IAudioExtensionPlugin.h"
 #include "OculusAudioMixer.h"
-#if PLATFORM_WINDOWS
-#include "OculusAudioLegacy.h"
-#endif
 #include "OculusAudioSourceSettings.h"
 #include "OculusAmbisonicSpatializer.h"
 
diff --git a/Engine/Plugins/Runtime/Oculus/OculusAudio/Source/OculusAudio/Private/OculusAudioEffect.cpp b/Engine/Plugins/Runtime/Oculus/OculusAudio/Source/OculusAudio/Private/OculusAudioEffect.cpp
deleted file mode 100644
diff --git a/Engine/Plugins/Runtime/Oculus/OculusAudio/Source/OculusAudio/Private/OculusAudioEffect.h b/Engine/Plugins/Runtime/Oculus/OculusAudio/Source/OculusAudio/Private/OculusAudioEffect.h
deleted file mode 100644
diff --git a/Engine/Plugins/Runtime/Oculus/OculusAudio/Source/OculusAudio/Private/OculusAudioLegacy.cpp b/Engine/Plugins/Runtime/Oculus/OculusAudio/Source/OculusAudio/Private/OculusAudioLegacy.cpp
deleted file mode 100644
diff --git a/Engine/Plugins/Runtime/Oculus/OculusAudio/Source/OculusAudio/Private/OculusAudioLegacy.h b/Engine/Plugins/Runtime/Oculus/OculusAudio/Source/OculusAudio/Private/OculusAudioLegacy.h
deleted file mode 100644
diff --git a/Engine/Plugins/Runtime/Steam/SteamAudio/Source/SteamAudio/SteamAudio.Build.cs b/Engine/Plugins/Runtime/Steam/SteamAudio/Source/SteamAudio/SteamAudio.Build.cs
--- a/Engine/Plugins/Runtime/Steam/SteamAudio/Source/SteamAudio/SteamAudio.Build.cs
+++ b/Engine/Plugins/Runtime/Steam/SteamAudio/Source/SteamAudio/SteamAudio.Build.cs
@@ -61,7 +61,6 @@ namespace UnrealBuildTool.Rules
 
 			if (Target.Platform == UnrealTargetPlatform.Win64)
 			{
-                PrivateDependencyModuleNames.Add("XAudio2");
 				AddEngineThirdPartyPrivateStaticDependencies(Target, "DX11Audio");
 
                 RuntimeDependencies.Add("$(EngineDir)/Binaries/ThirdParty/Phonon/Win64/...");
diff --git a/Engine/Source/Editor/UnrealEd/Classes/Commandlets/AudioTestCommandlet.h b/Engine/Source/Editor/UnrealEd/Classes/Commandlets/AudioTestCommandlet.h
deleted file mode 100644
diff --git a/Engine/Source/Editor/UnrealEd/Private/Commandlets/AudioTestCommandlet.cpp b/Engine/Source/Editor/UnrealEd/Private/Commandlets/AudioTestCommandlet.cpp
deleted file mode 100644
diff --git a/Engine/Source/Editor/UnrealEd/Public/UnrealEd.h b/Engine/Source/Editor/UnrealEd/Public/UnrealEd.h
--- a/Engine/Source/Editor/UnrealEd/Public/UnrealEd.h
+++ b/Engine/Source/Editor/UnrealEd/Public/UnrealEd.h
@@ -55,7 +55,6 @@ MONOLITHIC_HEADER_BOILERPLATE()
 #include "ActorFactories/ActorFactoryTriggerSphere.h"
 #include "ActorFactories/ActorFactoryVectorFieldVolume.h"
 #include "Exporters/AnimSequenceExporterFBX.h"
-#include "Commandlets/AudioTestCommandlet.h"
 #include "Preferences/CascadeOptions.h"
 #include "Settings/ClassViewerSettings.h"
 #include "Commandlets/CompressAnimationsCommandlet.h"
diff --git a/Engine/Source/Editor/UnrealEd/UnrealEd.Build.cs b/Engine/Source/Editor/UnrealEd/UnrealEd.Build.cs
--- a/Engine/Source/Editor/UnrealEd/UnrealEd.Build.cs
+++ b/Engine/Source/Editor/UnrealEd/UnrealEd.Build.cs
@@ -95,7 +95,6 @@ public class UnrealEd : ModuleRules
 				"GameplayDebugger",
 				"BlueprintGraph",
 				"HTTP",
-				"UnrealAudio",
 				"FunctionalTesting",
 				"AutomationController",
 				"Localization",
@@ -334,7 +333,6 @@ public class UnrealEd : ModuleRules
 
 		if (Target.Platform == UnrealTargetPlatform.Win64)
 		{
-			PublicDependencyModuleNames.Add("XAudio2");
 			PublicDependencyModuleNames.Add("AudioMixerXAudio2");
 
 			PrivateDependencyModuleNames.Add("WindowsPlatformFeatures");
diff --git a/Engine/Source/Programs/UnrealBuildTool/Platform/HoloLens/UEBuildHoloLens.cs b/Engine/Source/Programs/UnrealBuildTool/Platform/HoloLens/UEBuildHoloLens.cs
--- a/Engine/Source/Programs/UnrealBuildTool/Platform/HoloLens/UEBuildHoloLens.cs
+++ b/Engine/Source/Programs/UnrealBuildTool/Platform/HoloLens/UEBuildHoloLens.cs
@@ -511,14 +511,6 @@ namespace UnrealBuildTool
 				Rules.PublicAdditionalLibraries.Remove("XAPOFX.lib");
 				Rules.PrivateDependencyModuleNames.Add("WinPixEventRuntime");
 			}
-			else if (ModuleName == "XAudio2")
-			{
-				Rules.PublicDefinitions.Add("XAUDIO_SUPPORTS_XMA2WAVEFORMATEX=0");
-				Rules.PublicDefinitions.Add("XAUDIO_SUPPORTS_DEVICE_DETAILS=0");
-				Rules.PublicDefinitions.Add("XAUDIO2_SUPPORTS_MUSIC=0");
-				Rules.PublicDefinitions.Add("XAUDIO2_SUPPORTS_SENDLIST=1");
-				Rules.PublicSystemLibraries.Add("XAudio2.lib");
-			}
 			else if (ModuleName == "DX11Audio")
 			{
 				Rules.PublicAdditionalLibraries.Remove("X3DAudio.lib");
diff --git a/Engine/Source/Runtime/ALAudio/ALAudio.Build.cs b/Engine/Source/Runtime/ALAudio/ALAudio.Build.cs
deleted file mode 100644
diff --git a/Engine/Source/Runtime/ALAudio/Private/ALAudioBuffer.cpp b/Engine/Source/Runtime/ALAudio/Private/ALAudioBuffer.cpp
deleted file mode 100644
diff --git a/Engine/Source/Runtime/ALAudio/Private/ALAudioDevice.cpp b/Engine/Source/Runtime/ALAudio/Private/ALAudioDevice.cpp
deleted file mode 100644
diff --git a/Engine/Source/Runtime/ALAudio/Private/ALAudioSource.cpp b/Engine/Source/Runtime/ALAudio/Private/ALAudioSource.cpp
deleted file mode 100644
diff --git a/Engine/Source/Runtime/ALAudio/Public/ALAudioDevice.h b/Engine/Source/Runtime/ALAudio/Public/ALAudioDevice.h
deleted file mode 100644
diff --git a/Engine/Source/Runtime/Android/AndroidAudio/AndroidAudio.Build.cs b/Engine/Source/Runtime/Android/AndroidAudio/AndroidAudio.Build.cs
deleted file mode 100644
diff --git a/Engine/Source/Runtime/Android/AndroidAudio/Private/AndroidAudioBuffer.cpp b/Engine/Source/Runtime/Android/AndroidAudio/Private/AndroidAudioBuffer.cpp
deleted file mode 100644
diff --git a/Engine/Source/Runtime/Android/AndroidAudio/Private/AndroidAudioDevice.cpp b/Engine/Source/Runtime/Android/AndroidAudio/Private/AndroidAudioDevice.cpp
deleted file mode 100644
diff --git a/Engine/Source/Runtime/Android/AndroidAudio/Private/AndroidAudioSource.cpp b/Engine/Source/Runtime/Android/AndroidAudio/Private/AndroidAudioSource.cpp
deleted file mode 100644
diff --git a/Engine/Source/Runtime/Android/AndroidAudio/Public/AndroidAudioDevice.h b/Engine/Source/Runtime/Android/AndroidAudio/Public/AndroidAudioDevice.h
deleted file mode 100644
diff --git a/Engine/Source/Runtime/ApplicationCore/Private/Windows/WindowsPlatformApplicationMisc.cpp b/Engine/Source/Runtime/ApplicationCore/Private/Windows/WindowsPlatformApplicationMisc.cpp
--- a/Engine/Source/Runtime/ApplicationCore/Private/Windows/WindowsPlatformApplicationMisc.cpp
+++ b/Engine/Source/Runtime/ApplicationCore/Private/Windows/WindowsPlatformApplicationMisc.cpp
@@ -40,7 +40,6 @@ void FWindowsPlatformApplicationMisc::PreInit()
 void FWindowsPlatformApplicationMisc::LoadStartupModules()
 {
 #if !UE_SERVER
-	FModuleManager::Get().LoadModule(TEXT("XAudio2"));
 	FModuleManager::Get().LoadModule(TEXT("HeadMountedDisplay"));
 #endif // !UE_SERVER
 
diff --git a/Engine/Source/Runtime/Launch/Launch.Build.cs b/Engine/Source/Runtime/Launch/Launch.Build.cs
--- a/Engine/Source/Runtime/Launch/Launch.Build.cs
+++ b/Engine/Source/Runtime/Launch/Launch.Build.cs
@@ -81,22 +81,18 @@ public class Launch : ModuleRules
 			{
 				DynamicallyLoadedModuleNames.AddRange(new string[] {
 					"AudioMixerXAudio2",
-					"XAudio2",
 				});
 			}
 			else if (Target.Platform == UnrealTargetPlatform.HoloLens)
 			{
 				DynamicallyLoadedModuleNames.Add("D3D11RHI");
-				DynamicallyLoadedModuleNames.Add("XAudio2");
 				DynamicallyLoadedModuleNames.Add("AudioMixerXAudio2");
 			}
 			else if (Target.Platform == UnrealTargetPlatform.Mac)
 			{
 				DynamicallyLoadedModuleNames.AddRange(new string [] {
 					"AudioMixerCoreAudio",
-					"CoreAudio"
 				});
-				PublicFrameworks.Add("CoreAudio");
 			}
 			else if (Target.IsInPlatformGroup(UnrealPlatformGroup.Unix))
 			{
@@ -226,7 +222,6 @@ public class Launch : ModuleRules
 		if (Target.IsInPlatformGroup(UnrealPlatformGroup.Android))
 		{
 			PrivateDependencyModuleNames.Add("OpenGLDrv");
-			PrivateDependencyModuleNames.Add("AndroidAudio");
 			PrivateDependencyModuleNames.Add("AudioMixerAndroid");
 
 			// these are, for now, only for basic android
diff --git a/Engine/Source/Runtime/Mac/CoreAudio/CoreAudio.Build.cs b/Engine/Source/Runtime/Mac/CoreAudio/CoreAudio.Build.cs
deleted file mode 100755
diff --git a/Engine/Source/Runtime/Mac/CoreAudio/Private/CoreAudioBuffer.cpp b/Engine/Source/Runtime/Mac/CoreAudio/Private/CoreAudioBuffer.cpp
deleted file mode 100755
diff --git a/Engine/Source/Runtime/Mac/CoreAudio/Private/CoreAudioDevice.cpp b/Engine/Source/Runtime/Mac/CoreAudio/Private/CoreAudioDevice.cpp
deleted file mode 100755
diff --git a/Engine/Source/Runtime/Mac/CoreAudio/Private/CoreAudioEffects.cpp b/Engine/Source/Runtime/Mac/CoreAudio/Private/CoreAudioEffects.cpp
deleted file mode 100755
diff --git a/Engine/Source/Runtime/Mac/CoreAudio/Private/CoreAudioSource.cpp b/Engine/Source/Runtime/Mac/CoreAudio/Private/CoreAudioSource.cpp
deleted file mode 100755
diff --git a/Engine/Source/Runtime/Mac/CoreAudio/Public/CoreAudioDevice.h b/Engine/Source/Runtime/Mac/CoreAudio/Public/CoreAudioDevice.h
deleted file mode 100755
diff --git a/Engine/Source/Runtime/Mac/CoreAudio/Public/CoreAudioEffects.h b/Engine/Source/Runtime/Mac/CoreAudio/Public/CoreAudioEffects.h
deleted file mode 100755
diff --git a/Engine/Source/Runtime/Mac/UnrealAudioCoreAudio/Private/UnrealAudioCoreAudio.cpp b/Engine/Source/Runtime/Mac/UnrealAudioCoreAudio/Private/UnrealAudioCoreAudio.cpp
deleted file mode 100644
diff --git a/Engine/Source/Runtime/Mac/UnrealAudioCoreAudio/UnrealAudioCoreAudio.Build.cs b/Engine/Source/Runtime/Mac/UnrealAudioCoreAudio/UnrealAudioCoreAudio.Build.cs
deleted file mode 100644
diff --git a/Engine/Source/Runtime/UnrealAudio/Private/Tests/UnrealAudioDeviceTests.cpp b/Engine/Source/Runtime/UnrealAudio/Private/Tests/UnrealAudioDeviceTests.cpp
deleted file mode 100644
diff --git a/Engine/Source/Runtime/UnrealAudio/Private/Tests/UnrealAudioSystemTests.cpp b/Engine/Source/Runtime/UnrealAudio/Private/Tests/UnrealAudioSystemTests.cpp
deleted file mode 100644
diff --git a/Engine/Source/Runtime/UnrealAudio/Private/Tests/UnrealAudioTestGenerators.cpp b/Engine/Source/Runtime/UnrealAudio/Private/Tests/UnrealAudioTestGenerators.cpp
deleted file mode 100644
diff --git a/Engine/Source/Runtime/UnrealAudio/Private/Tests/UnrealAudioTestGenerators.h b/Engine/Source/Runtime/UnrealAudio/Private/Tests/UnrealAudioTestGenerators.h
deleted file mode 100644
diff --git a/Engine/Source/Runtime/UnrealAudio/Private/UnrealAudio.cpp b/Engine/Source/Runtime/UnrealAudio/Private/UnrealAudio.cpp
deleted file mode 100644
diff --git a/Engine/Source/Runtime/UnrealAudio/Private/UnrealAudioBuffer.cpp b/Engine/Source/Runtime/UnrealAudio/Private/UnrealAudioBuffer.cpp
deleted file mode 100644
diff --git a/Engine/Source/Runtime/UnrealAudio/Private/UnrealAudioBuffer.h b/Engine/Source/Runtime/UnrealAudio/Private/UnrealAudioBuffer.h
deleted file mode 100644
diff --git a/Engine/Source/Runtime/UnrealAudio/Private/UnrealAudioDecode.cpp b/Engine/Source/Runtime/UnrealAudio/Private/UnrealAudioDecode.cpp
deleted file mode 100644
diff --git a/Engine/Source/Runtime/UnrealAudio/Private/UnrealAudioDecode.h b/Engine/Source/Runtime/UnrealAudio/Private/UnrealAudioDecode.h
deleted file mode 100644
diff --git a/Engine/Source/Runtime/UnrealAudio/Private/UnrealAudioDevice.cpp b/Engine/Source/Runtime/UnrealAudio/Private/UnrealAudioDevice.cpp
deleted file mode 100644
diff --git a/Engine/Source/Runtime/UnrealAudio/Private/UnrealAudioDeviceFormat.cpp b/Engine/Source/Runtime/UnrealAudio/Private/UnrealAudioDeviceFormat.cpp
deleted file mode 100644
diff --git a/Engine/Source/Runtime/UnrealAudio/Private/UnrealAudioDeviceFormat.h b/Engine/Source/Runtime/UnrealAudio/Private/UnrealAudioDeviceFormat.h
deleted file mode 100644
diff --git a/Engine/Source/Runtime/UnrealAudio/Private/UnrealAudioEmitter.cpp b/Engine/Source/Runtime/UnrealAudio/Private/UnrealAudioEmitter.cpp
deleted file mode 100644
diff --git a/Engine/Source/Runtime/UnrealAudio/Private/UnrealAudioEmitterInternal.h b/Engine/Source/Runtime/UnrealAudio/Private/UnrealAudioEmitterInternal.h
deleted file mode 100644
diff --git a/Engine/Source/Runtime/UnrealAudio/Private/UnrealAudioEmitterManager.cpp b/Engine/Source/Runtime/UnrealAudio/Private/UnrealAudioEmitterManager.cpp
deleted file mode 100644
diff --git a/Engine/Source/Runtime/UnrealAudio/Private/UnrealAudioEmitterManager.h b/Engine/Source/Runtime/UnrealAudio/Private/UnrealAudioEmitterManager.h
deleted file mode 100644
diff --git a/Engine/Source/Runtime/UnrealAudio/Private/UnrealAudioEntityManager.h b/Engine/Source/Runtime/UnrealAudio/Private/UnrealAudioEntityManager.h
deleted file mode 100644
diff --git a/Engine/Source/Runtime/UnrealAudio/Private/UnrealAudioHandles.h b/Engine/Source/Runtime/UnrealAudio/Private/UnrealAudioHandles.h
deleted file mode 100644
diff --git a/Engine/Source/Runtime/UnrealAudio/Private/UnrealAudioPitchManager.cpp b/Engine/Source/Runtime/UnrealAudio/Private/UnrealAudioPitchManager.cpp
deleted file mode 100644
diff --git a/Engine/Source/Runtime/UnrealAudio/Private/UnrealAudioPitchManager.h b/Engine/Source/Runtime/UnrealAudio/Private/UnrealAudioPitchManager.h
deleted file mode 100644
diff --git a/Engine/Source/Runtime/UnrealAudio/Private/UnrealAudioPrivate.h b/Engine/Source/Runtime/UnrealAudio/Private/UnrealAudioPrivate.h
deleted file mode 100644
diff --git a/Engine/Source/Runtime/UnrealAudio/Private/UnrealAudioSampleRateConverter.cpp b/Engine/Source/Runtime/UnrealAudio/Private/UnrealAudioSampleRateConverter.cpp
deleted file mode 100644
diff --git a/Engine/Source/Runtime/UnrealAudio/Private/UnrealAudioSampleRateConverter.h b/Engine/Source/Runtime/UnrealAudio/Private/UnrealAudioSampleRateConverter.h
deleted file mode 100644
diff --git a/Engine/Source/Runtime/UnrealAudio/Private/UnrealAudioSoundFile.cpp b/Engine/Source/Runtime/UnrealAudio/Private/UnrealAudioSoundFile.cpp
deleted file mode 100644
diff --git a/Engine/Source/Runtime/UnrealAudio/Private/UnrealAudioSoundFileConvert.cpp b/Engine/Source/Runtime/UnrealAudio/Private/UnrealAudioSoundFileConvert.cpp
deleted file mode 100644
diff --git a/Engine/Source/Runtime/UnrealAudio/Private/UnrealAudioSoundFileInternal.h b/Engine/Source/Runtime/UnrealAudio/Private/UnrealAudioSoundFileInternal.h
deleted file mode 100644
diff --git a/Engine/Source/Runtime/UnrealAudio/Private/UnrealAudioSoundFileManager.cpp b/Engine/Source/Runtime/UnrealAudio/Private/UnrealAudioSoundFileManager.cpp
deleted file mode 100644
diff --git a/Engine/Source/Runtime/UnrealAudio/Private/UnrealAudioSoundFileManager.h b/Engine/Source/Runtime/UnrealAudio/Private/UnrealAudioSoundFileManager.h
deleted file mode 100644
diff --git a/Engine/Source/Runtime/UnrealAudio/Private/UnrealAudioSoundFileReader.cpp b/Engine/Source/Runtime/UnrealAudio/Private/UnrealAudioSoundFileReader.cpp
deleted file mode 100644
diff --git a/Engine/Source/Runtime/UnrealAudio/Private/UnrealAudioThread.h b/Engine/Source/Runtime/UnrealAudio/Private/UnrealAudioThread.h
deleted file mode 100644
diff --git a/Engine/Source/Runtime/UnrealAudio/Private/UnrealAudioUtilities.cpp b/Engine/Source/Runtime/UnrealAudio/Private/UnrealAudioUtilities.cpp
deleted file mode 100644
diff --git a/Engine/Source/Runtime/UnrealAudio/Private/UnrealAudioUtilities.h b/Engine/Source/Runtime/UnrealAudio/Private/UnrealAudioUtilities.h
deleted file mode 100644
diff --git a/Engine/Source/Runtime/UnrealAudio/Private/UnrealAudioVoice.cpp b/Engine/Source/Runtime/UnrealAudio/Private/UnrealAudioVoice.cpp
deleted file mode 100644
diff --git a/Engine/Source/Runtime/UnrealAudio/Private/UnrealAudioVoiceInternal.h b/Engine/Source/Runtime/UnrealAudio/Private/UnrealAudioVoiceInternal.h
deleted file mode 100644
diff --git a/Engine/Source/Runtime/UnrealAudio/Private/UnrealAudioVoiceManager.cpp b/Engine/Source/Runtime/UnrealAudio/Private/UnrealAudioVoiceManager.cpp
deleted file mode 100644
diff --git a/Engine/Source/Runtime/UnrealAudio/Private/UnrealAudioVoiceManager.h b/Engine/Source/Runtime/UnrealAudio/Private/UnrealAudioVoiceManager.h
deleted file mode 100644
diff --git a/Engine/Source/Runtime/UnrealAudio/Private/UnrealAudioVoiceMixer.cpp b/Engine/Source/Runtime/UnrealAudio/Private/UnrealAudioVoiceMixer.cpp
deleted file mode 100644
diff --git a/Engine/Source/Runtime/UnrealAudio/Private/UnrealAudioVoiceMixer.h b/Engine/Source/Runtime/UnrealAudio/Private/UnrealAudioVoiceMixer.h
deleted file mode 100644
diff --git a/Engine/Source/Runtime/UnrealAudio/Private/UnrealAudioVolumeManager.cpp b/Engine/Source/Runtime/UnrealAudio/Private/UnrealAudioVolumeManager.cpp
deleted file mode 100644
diff --git a/Engine/Source/Runtime/UnrealAudio/Private/UnrealAudioVolumeManager.h b/Engine/Source/Runtime/UnrealAudio/Private/UnrealAudioVolumeManager.h
deleted file mode 100644
diff --git a/Engine/Source/Runtime/UnrealAudio/Public/UnrealAudioDeviceModule.h b/Engine/Source/Runtime/UnrealAudio/Public/UnrealAudioDeviceModule.h
deleted file mode 100644
diff --git a/Engine/Source/Runtime/UnrealAudio/Public/UnrealAudioEmitter.h b/Engine/Source/Runtime/UnrealAudio/Public/UnrealAudioEmitter.h
deleted file mode 100644
diff --git a/Engine/Source/Runtime/UnrealAudio/Public/UnrealAudioModule.h b/Engine/Source/Runtime/UnrealAudio/Public/UnrealAudioModule.h
deleted file mode 100644
diff --git a/Engine/Source/Runtime/UnrealAudio/Public/UnrealAudioSoundFile.h b/Engine/Source/Runtime/UnrealAudio/Public/UnrealAudioSoundFile.h
deleted file mode 100644
diff --git a/Engine/Source/Runtime/UnrealAudio/Public/UnrealAudioTests.h b/Engine/Source/Runtime/UnrealAudio/Public/UnrealAudioTests.h
deleted file mode 100644
diff --git a/Engine/Source/Runtime/UnrealAudio/Public/UnrealAudioTypes.h b/Engine/Source/Runtime/UnrealAudio/Public/UnrealAudioTypes.h
deleted file mode 100644
diff --git a/Engine/Source/Runtime/UnrealAudio/Public/UnrealAudioVoice.h b/Engine/Source/Runtime/UnrealAudio/Public/UnrealAudioVoice.h
deleted file mode 100644
diff --git a/Engine/Source/Runtime/UnrealAudio/UnrealAudio.Build.cs b/Engine/Source/Runtime/UnrealAudio/UnrealAudio.Build.cs
deleted file mode 100644
diff --git a/Engine/Source/Runtime/Windows/UnrealAudioWasapi/Private/UnrealAudioDeviceWasapi.cpp b/Engine/Source/Runtime/Windows/UnrealAudioWasapi/Private/UnrealAudioDeviceWasapi.cpp
deleted file mode 100644
diff --git a/Engine/Source/Runtime/Windows/UnrealAudioWasapi/UnrealAudioWasapi.Build.cs b/Engine/Source/Runtime/Windows/UnrealAudioWasapi/UnrealAudioWasapi.Build.cs
deleted file mode 100644
diff --git a/Engine/Source/Runtime/Windows/UnrealAudioXAudio2/Private/UnrealAudioDeviceXAudio2.cpp b/Engine/Source/Runtime/Windows/UnrealAudioXAudio2/Private/UnrealAudioDeviceXAudio2.cpp
deleted file mode 100644
diff --git a/Engine/Source/Runtime/Windows/UnrealAudioXAudio2/UnrealAudioXAudio2.Build.cs b/Engine/Source/Runtime/Windows/UnrealAudioXAudio2/UnrealAudioXAudio2.Build.cs
deleted file mode 100644
diff --git a/Engine/Source/Runtime/Windows/XAudio2/Private/XAudio2Buffer.cpp b/Engine/Source/Runtime/Windows/XAudio2/Private/XAudio2Buffer.cpp
deleted file mode 100644
diff --git a/Engine/Source/Runtime/Windows/XAudio2/Private/XAudio2Device.cpp b/Engine/Source/Runtime/Windows/XAudio2/Private/XAudio2Device.cpp
deleted file mode 100644
diff --git a/Engine/Source/Runtime/Windows/XAudio2/Private/XAudio2Effects.cpp b/Engine/Source/Runtime/Windows/XAudio2/Private/XAudio2Effects.cpp
deleted file mode 100644
diff --git a/Engine/Source/Runtime/Windows/XAudio2/Private/XAudio2Source.cpp b/Engine/Source/Runtime/Windows/XAudio2/Private/XAudio2Source.cpp
deleted file mode 100644
diff --git a/Engine/Source/Runtime/Windows/XAudio2/Private/XAudio2Support.h b/Engine/Source/Runtime/Windows/XAudio2/Private/XAudio2Support.h
deleted file mode 100644
diff --git a/Engine/Source/Runtime/Windows/XAudio2/Public/XAudio2Device.h b/Engine/Source/Runtime/Windows/XAudio2/Public/XAudio2Device.h
deleted file mode 100644
diff --git a/Engine/Source/Runtime/Windows/XAudio2/Public/XAudio2Effects.h b/Engine/Source/Runtime/Windows/XAudio2/Public/XAudio2Effects.h
deleted file mode 100644
diff --git a/Engine/Source/Runtime/Windows/XAudio2/XAudio2.Build.cs b/Engine/Source/Runtime/Windows/XAudio2/XAudio2.Build.cs
deleted file mode 100644
