Remove platform references.

diff --git a/Engine/Plugins/Online/OnlineSubsystemUtils/Source/OnlineSubsystemUtils/Classes/IpNetDriver.h b/Engine/Plugins/Online/OnlineSubsystemUtils/Source/OnlineSubsystemUtils/Classes/IpNetDriver.h
--- a/Engine/Plugins/Online/OnlineSubsystemUtils/Source/OnlineSubsystemUtils/Classes/IpNetDriver.h
+++ b/Engine/Plugins/Online/OnlineSubsystemUtils/Source/OnlineSubsystemUtils/Classes/IpNetDriver.h
@@ -218,7 +218,7 @@ private:
 	 */
 	static FORCEINLINE bool IsRecvFailBlocking(ESocketErrors Error)
 	{
-		// SE_ECONNABORTED is for PS4 LAN cable pulls, SE_ENETDOWN is for a Switch hang
+		// SE_ECONNABORTED is for LAN cable pulls on some platforms, SE_ENETDOWN is to prevent a hang on another platform
 		return Error == SE_NO_ERROR || Error == SE_EWOULDBLOCK || Error == SE_ECONNABORTED || Error == SE_ENETDOWN;
 	};
 
