Check for uobject system being valid before accessing voice components on shutdown

diff --git a/Engine/Plugins/Online/OnlineSubsystemUtils/Source/OnlineSubsystemUtils/Private/VoiceEngineImpl.cpp b/Engine/Plugins/Online/OnlineSubsystemUtils/Source/OnlineSubsystemUtils/Private/VoiceEngineImpl.cpp
--- a/Engine/Plugins/Online/OnlineSubsystemUtils/Source/OnlineSubsystemUtils/Private/VoiceEngineImpl.cpp
+++ b/Engine/Plugins/Online/OnlineSubsystemUtils/Source/OnlineSubsystemUtils/Private/VoiceEngineImpl.cpp
@@ -122,7 +122,7 @@ void FRemoteTalkerDataImpl::Reset()
 	LastSeen = MAX_FLT;
 	NumFramesStarved = 0;
 
-	if (VoipSynthComponent)
+	if (UObjectInitialized() && VoipSynthComponent)
 	{
 		VoipSynthComponent->Stop();
 
