Allow duplicate PIE login credentials if the type is Developer

diff --git a/Engine/Plugins/Online/OnlineSubsystemUtils/Source/OnlineSubsystemUtils/Private/OnlinePIESettings.cpp b/Engine/Plugins/Online/OnlineSubsystemUtils/Source/OnlineSubsystemUtils/Private/OnlinePIESettings.cpp
--- a/Engine/Plugins/Online/OnlineSubsystemUtils/Source/OnlineSubsystemUtils/Private/OnlinePIESettings.cpp
+++ b/Engine/Plugins/Online/OnlineSubsystemUtils/Source/OnlineSubsystemUtils/Private/OnlinePIESettings.cpp
@@ -125,7 +125,10 @@ void UOnlinePIESettings::PostEditChangeProperty(FPropertyChangedEvent& PropertyC
 					// Remove any whitespace from login input
 					Login.Id.TrimStartAndEndInline();
 
-					if (Ids.Contains(Login.Id))
+					const bool bIsExistingLogin = Ids.Contains(Login.Id);
+					// Allow duplicate for Developer credentials.  This is intended to support logins where other fields uniquely identify the login but require the same id.
+					const bool bDuplicateAllowed = Login.Type.Equals(TEXT("Developer"), ESearchCase::IgnoreCase);
+					if (bIsExistingLogin && !bDuplicateAllowed)
 					{
 						// Don't allow duplicate login ids
 						Login.Id.Reset();
