Add PIE support for online services (OSSv2)
Deprecated OnlineEngineInterface methods that take FUniqueNetId and add new methods that work with FUniqueNetIdWrapper.  FUniqueNetIdWrapper supports FUniqueNetId and OSSv2 ids.

diff --git a/Engine/Plugins/Online/OnlineSubsystemUtils/OnlineSubsystemUtils.uplugin b/Engine/Plugins/Online/OnlineSubsystemUtils/OnlineSubsystemUtils.uplugin
--- a/Engine/Plugins/Online/OnlineSubsystemUtils/OnlineSubsystemUtils.uplugin
+++ b/Engine/Plugins/Online/OnlineSubsystemUtils/OnlineSubsystemUtils.uplugin
@@ -4,7 +4,7 @@
 	"FriendlyName" : "Online Subsystem Utils",
 	"Version" : 1,
 	"VersionName" : "1.0",
-	"Description" : "Shared code for interacting online subsystem implementations.",
+	"Description" : "Shared code for interacting online service and online subsystem implementations.",
 	"Category" : "Online Platform",
 	"CreatedBy" : "Epic Games, Inc.",
 	"CreatedByURL" : "http://epicgames.com",
@@ -34,6 +34,10 @@
 		{
 			"Name": "OnlineSubsystem",
 			"Enabled": true
+		},
+		{
+			"Name": "OnlineServices",
+			"Enabled": true
 		}
 	]
 }
\ No newline at end of file
diff --git a/Engine/Plugins/Online/OnlineSubsystemUtils/Source/OnlineSubsystemUtils/OnlineSubsystemUtils.Build.cs b/Engine/Plugins/Online/OnlineSubsystemUtils/Source/OnlineSubsystemUtils/OnlineSubsystemUtils.Build.cs
--- a/Engine/Plugins/Online/OnlineSubsystemUtils/Source/OnlineSubsystemUtils/OnlineSubsystemUtils.Build.cs
+++ b/Engine/Plugins/Online/OnlineSubsystemUtils/Source/OnlineSubsystemUtils/OnlineSubsystemUtils.Build.cs
@@ -46,7 +46,8 @@ public class OnlineSubsystemUtils : ModuleRules
 				"Json",
 				"SignalProcessing",
 				"AudioMixerCore",
-				"DeveloperSettings"
+				"DeveloperSettings",
+				"OnlineServicesInterface"
 			}
 		);
 
diff --git a/Engine/Plugins/Online/OnlineSubsystemUtils/Source/OnlineSubsystemUtils/Private/OnlineEngineInterfaceImpl.cpp b/Engine/Plugins/Online/OnlineSubsystemUtils/Source/OnlineSubsystemUtils/Private/OnlineEngineInterfaceImpl.cpp
--- a/Engine/Plugins/Online/OnlineSubsystemUtils/Source/OnlineSubsystemUtils/Private/OnlineEngineInterfaceImpl.cpp
+++ b/Engine/Plugins/Online/OnlineSubsystemUtils/Source/OnlineSubsystemUtils/Private/OnlineEngineInterfaceImpl.cpp
@@ -68,9 +68,9 @@ FName UOnlineEngineInterfaceImpl::GetDefaultOnlineSubsystemName() const
 	return OnlineSub ? OnlineSub->GetSubsystemName() : NAME_None;
 }
 
-bool UOnlineEngineInterfaceImpl::IsCompatibleUniqueNetId(const FUniqueNetId& InUniqueNetId) const
+bool UOnlineEngineInterfaceImpl::IsCompatibleUniqueNetId(const FUniqueNetIdWrapper& InUniqueNetId) const
 {
-	return InUniqueNetId.GetType() == GetDefaultOnlineSubsystemName() || CompatibleUniqueNetIdTypes.Contains(InUniqueNetId.GetType());
+	return InUniqueNetId.IsV1() && (InUniqueNetId.GetType() == GetDefaultOnlineSubsystemName() || CompatibleUniqueNetIdTypes.Contains(InUniqueNetId.GetType()));
 }
 
 uint8 UOnlineEngineInterfaceImpl::GetReplicationHashForSubsystem(FName InSubsystemName) const
@@ -109,7 +109,7 @@ PRAGMA_ENABLE_DEPRECATION_WARNINGS
 	return NAME_None;
 }
 
-FUniqueNetIdPtr UOnlineEngineInterfaceImpl::CreateUniquePlayerId(const FString& Str, FName Type)
+FUniqueNetIdWrapper UOnlineEngineInterfaceImpl::CreateUniquePlayerIdWrapper(const FString& Str, FName Type)
 {
 	// Foreign types may be passed into this function, do not load OSS modules explicitly here
 	FUniqueNetIdPtr UniqueId = nullptr;
@@ -143,7 +143,7 @@ FUniqueNetIdPtr UOnlineEngineInterfaceImpl::CreateUniquePlayerId(const FString&
 	return UniqueId;
 }
 
-FUniqueNetIdPtr UOnlineEngineInterfaceImpl::GetUniquePlayerId(UWorld* World, int32 LocalUserNum, FName Type)
+FUniqueNetIdWrapper UOnlineEngineInterfaceImpl::GetUniquePlayerIdWrapper(UWorld* World, int32 LocalUserNum, FName Type)
 {
 	IOnlineIdentityPtr IdentityInt = Online::GetIdentityInterface(World, Type);
 	if (IdentityInt.IsValid())
@@ -153,15 +153,16 @@ FUniqueNetIdPtr UOnlineEngineInterfaceImpl::GetUniquePlayerId(UWorld* World, int
 	}
 
 	UE_LOG_ONLINE(Verbose, TEXT("GetUniquePlayerId() returning null, can't find OSS of type %s"), *Type.ToString());
-	return nullptr;
+	return FUniqueNetIdWrapper();
 }
 
-FString UOnlineEngineInterfaceImpl::GetPlayerNickname(UWorld* World, const FUniqueNetId& UniqueId)
+FString UOnlineEngineInterfaceImpl::GetPlayerNickname(UWorld* World, const FUniqueNetIdWrapper& UniqueId)
 {
+	check(UniqueId.IsValid() && UniqueId.IsV1());
 	IOnlineIdentityPtr IdentityInt = Online::GetIdentityInterface(World, UniqueId.GetType());
 	if (IdentityInt.IsValid())
 	{	
-		return IdentityInt->GetPlayerNickname(UniqueId);
+		return IdentityInt->GetPlayerNickname(*UniqueId);
 	}
 
 	static FString InvalidName(TEXT("InvalidOSSUser"));
@@ -353,30 +354,38 @@ void UOnlineEngineInterfaceImpl::UpdateSessionJoinability(UWorld* World, FName S
 	}
 }
 
-void UOnlineEngineInterfaceImpl::RegisterPlayer(UWorld* World, FName SessionName, const FUniqueNetId& UniqueId, bool bWasInvited)
+void UOnlineEngineInterfaceImpl::RegisterPlayer(UWorld* World, FName SessionName, const FUniqueNetIdWrapper& UniqueId, bool bWasInvited)
 {
+	check(UniqueId.IsValid() && UniqueId.IsV1());
 	IOnlineSessionPtr SessionInt = Online::GetSessionInterface(World);
 	if (SessionInt.IsValid() && UniqueId.IsValid())
 	{
-		SessionInt->RegisterPlayer(SessionName, UniqueId, bWasInvited);
+		SessionInt->RegisterPlayer(SessionName, *UniqueId, bWasInvited);
 	}
 }
 
-void UOnlineEngineInterfaceImpl::UnregisterPlayer(UWorld* World, FName SessionName, const FUniqueNetId& UniqueId)
+void UOnlineEngineInterfaceImpl::UnregisterPlayer(UWorld* World, FName SessionName, const FUniqueNetIdWrapper& UniqueId)
 {
+	check(UniqueId.IsValid() && UniqueId.IsV1());
 	IOnlineSessionPtr SessionInt = Online::GetSessionInterface(World);
 	if (SessionInt.IsValid())
 	{
-		SessionInt->UnregisterPlayer(SessionName, UniqueId);
+		SessionInt->UnregisterPlayer(SessionName, *UniqueId);
 	}
 }
 
-void UOnlineEngineInterfaceImpl::UnregisterPlayers(UWorld* World, FName SessionName, const TArray< FUniqueNetIdRef >& Players)
+void UOnlineEngineInterfaceImpl::UnregisterPlayers(UWorld* World, FName SessionName, const TArray<FUniqueNetIdWrapper>& Players)
 {
 	IOnlineSessionPtr SessionInt = Online::GetSessionInterface(World);
 	if (SessionInt.IsValid())
 	{
-		SessionInt->UnregisterPlayers(SessionName, Players);
+		TArray<FUniqueNetIdRef> PlayerIdsAsRefs;
+		for (const FUniqueNetIdWrapper& PlayerId : Players)
+		{
+			check(PlayerId.IsValid() && PlayerId.IsV1());
+			PlayerIdsAsRefs.Emplace(PlayerId.GetV1().ToSharedRef());
+		}
+		SessionInt->UnregisterPlayers(SessionName, PlayerIdsAsRefs);
 	}
 }
 
@@ -451,22 +460,24 @@ void UOnlineEngineInterfaceImpl::ClearVoicePackets(UWorld* World)
 	}
 }
 
-bool UOnlineEngineInterfaceImpl::MuteRemoteTalker(UWorld* World, uint8 LocalUserNum, const FUniqueNetId& PlayerId, bool bIsSystemWide)
+bool UOnlineEngineInterfaceImpl::MuteRemoteTalker(UWorld* World, uint8 LocalUserNum, const FUniqueNetIdWrapper& PlayerId, bool bIsSystemWide)
 {
+	check(PlayerId.IsValid() && PlayerId.IsV1());
 	IOnlineVoicePtr VoiceInt = Online::GetVoiceInterface(World, VoiceSubsystemNameOverride);
 	if (VoiceInt.IsValid())
 	{
-		return VoiceInt->MuteRemoteTalker(LocalUserNum, PlayerId, bIsSystemWide);
+		return VoiceInt->MuteRemoteTalker(LocalUserNum, *PlayerId, bIsSystemWide);
 	}
 	return false;
 }
 
-bool UOnlineEngineInterfaceImpl::UnmuteRemoteTalker(UWorld* World, uint8 LocalUserNum, const FUniqueNetId& PlayerId, bool bIsSystemWide)
+bool UOnlineEngineInterfaceImpl::UnmuteRemoteTalker(UWorld* World, uint8 LocalUserNum, const FUniqueNetIdWrapper& PlayerId, bool bIsSystemWide)
 {
+	check(PlayerId.IsValid() && PlayerId.IsV1());
 	IOnlineVoicePtr VoiceInt = Online::GetVoiceInterface(World, VoiceSubsystemNameOverride);
 	if (VoiceInt.IsValid())
 	{
-		return VoiceInt->UnmuteRemoteTalker(LocalUserNum, PlayerId, bIsSystemWide);
+		return VoiceInt->UnmuteRemoteTalker(LocalUserNum, *PlayerId, bIsSystemWide);
 	}
 	return false;
 }
diff --git a/Engine/Plugins/Online/OnlineSubsystemUtils/Source/OnlineSubsystemUtils/Private/OnlineEngineInterfaceImpl.h b/Engine/Plugins/Online/OnlineSubsystemUtils/Source/OnlineSubsystemUtils/Private/OnlineEngineInterfaceImpl.h
--- a/Engine/Plugins/Online/OnlineSubsystemUtils/Source/OnlineSubsystemUtils/Private/OnlineEngineInterfaceImpl.h
+++ b/Engine/Plugins/Online/OnlineSubsystemUtils/Source/OnlineSubsystemUtils/Private/OnlineEngineInterfaceImpl.h
@@ -28,7 +28,7 @@ public:
 	virtual void ShutdownOnlineSubsystem(FName OnlineIdentifier) override;
 	virtual void DestroyOnlineSubsystem(FName OnlineIdentifier) override;
 	virtual FName GetDefaultOnlineSubsystemName() const override;
-	virtual bool IsCompatibleUniqueNetId(const FUniqueNetId& InUniqueNetId) const override;
+	virtual bool IsCompatibleUniqueNetId(const FUniqueNetIdWrapper& InUniqueNetId) const override;
 
 	/**
 	 * Utils
@@ -60,10 +60,10 @@ private:
 	 */
 public:
 
-	virtual FUniqueNetIdPtr CreateUniquePlayerId(const FString& Str, FName Type) override;
-	virtual FUniqueNetIdPtr GetUniquePlayerId(UWorld* World, int32 LocalUserNum, FName Type) override;
+	virtual FUniqueNetIdWrapper CreateUniquePlayerIdWrapper(const FString& Str, FName Type) override;
+	virtual FUniqueNetIdWrapper GetUniquePlayerIdWrapper(UWorld* World, int32 LocalUserNum, FName Type) override;
 
-	virtual FString GetPlayerNickname(UWorld* World, const FUniqueNetId& UniqueId) override;
+	virtual FString GetPlayerNickname(UWorld* World, const FUniqueNetIdWrapper& UniqueId) override;
 	virtual bool GetPlayerPlatformNickname(UWorld* World, int32 LocalUserNum, FString& OutNickname) override;
 
 	virtual bool AutoLogin(UWorld* World, int32 LocalUserNum, const FOnlineAutoLoginComplete& InCompletionDelegate) override;
@@ -86,9 +86,9 @@ public:
 	virtual bool GetSessionJoinability(UWorld* World, FName SessionName, FJoinabilitySettings& OutSettings) override;
 	virtual void UpdateSessionJoinability(UWorld* World, FName SessionName, bool bPublicSearchable, bool bAllowInvites, bool bJoinViaPresence, bool bJoinViaPresenceFriendsOnly) override;
 
-	virtual void RegisterPlayer(UWorld* World, FName SessionName, const FUniqueNetId& UniqueId, bool bWasInvited) override;
-	virtual void UnregisterPlayer(UWorld* World, FName SessionName, const FUniqueNetId& UniqueId) override;
-	virtual void UnregisterPlayers(UWorld* World, FName SessionName, const TArray< FUniqueNetIdRef >& Players) override;
+	virtual void RegisterPlayer(UWorld* World, FName SessionName, const FUniqueNetIdWrapper& UniqueId, bool bWasInvited) override;
+	virtual void UnregisterPlayer(UWorld* World, FName SessionName, const FUniqueNetIdWrapper& UniqueId) override;
+	virtual void UnregisterPlayers(UWorld* World, FName SessionName, const TArray<FUniqueNetIdWrapper>& Players) override;
 
 	virtual bool GetResolvedConnectString(UWorld* World, FName SessionName, FString& URL) override;
 
@@ -114,8 +114,8 @@ public:
 	virtual void StopNetworkedVoice(UWorld* World, uint8 LocalUserNum) override;
 	virtual void ClearVoicePackets(UWorld* World) override;
 
-	virtual bool MuteRemoteTalker(UWorld* World, uint8 LocalUserNum, const FUniqueNetId& PlayerId, bool bIsSystemWide) override;
-	virtual bool UnmuteRemoteTalker(UWorld* World, uint8 LocalUserNum, const FUniqueNetId& PlayerId, bool bIsSystemWide) override;
+	virtual bool MuteRemoteTalker(UWorld* World, uint8 LocalUserNum, const FUniqueNetIdWrapper& PlayerId, bool bIsSystemWide) override;
+	virtual bool UnmuteRemoteTalker(UWorld* World, uint8 LocalUserNum, const FUniqueNetIdWrapper& PlayerId, bool bIsSystemWide) override;
 
 	virtual int32 GetNumLocalTalkers(UWorld* World) override;
 
@@ -127,10 +127,8 @@ public:
 	virtual void ShowLeaderboardUI(UWorld* World, const FString& CategoryName) override;
 	virtual void ShowAchievementsUI(UWorld* World, int32 LocalUserNum) override;
 	virtual void BindToExternalUIOpening(const FOnlineExternalUIChanged& Delegate) override;
-#ifdef OSS_ADDED_SHOW_WEB
 	virtual void ShowWebURL(const FString& CurrentURL, const UOnlineEngineInterface::FShowWebUrlParams& ShowParams, const FOnlineShowWebUrlClosed& CompletionDelegate) override;
 	virtual bool CloseWebURL() override;
-#endif
 
 private:
 
@@ -164,6 +162,5 @@ private:
 	/** Mapping of delegate handles for each online Login() call while in flight */
 	TMap<FName, FDelegateHandle> OnLoginPIECompleteDelegateHandlesForPIEInstances;
 	void OnPIELoginComplete(int32 LocalUserNum, bool bWasSuccessful, const FUniqueNetId& UserId, const FString& Error, FName OnlineIdentifier, FOnlineAutoLoginComplete InCompletionDelegate);
-
 };
 
diff --git a/Engine/Plugins/Online/OnlineSubsystemUtils/Source/OnlineSubsystemUtils/Private/OnlineServicesEngineInterfaceImpl.cpp b/Engine/Plugins/Online/OnlineSubsystemUtils/Source/OnlineSubsystemUtils/Private/OnlineServicesEngineInterfaceImpl.cpp
new file mode 100644
--- /dev/null
+++ b/Engine/Plugins/Online/OnlineSubsystemUtils/Source/OnlineSubsystemUtils/Private/OnlineServicesEngineInterfaceImpl.cpp
@@ -0,0 +1,480 @@
+// Copyright Epic Games, Inc. All Rights Reserved.
+
+#include "OnlineServicesEngineInterfaceImpl.h"
+#include "OnlineServicesEngineUtils.h"
+#include "Misc/ConfigCacheIni.h"
+#include "Engine/NetConnection.h"
+
+#include "Online/Auth.h"
+#include "Online/CoreOnline.h"
+#include "Online/ExternalUI.h"
+#include "Online/OnlineAsyncOpHandle.h"
+#include "Online/OnlineServices.h"
+#include "Online/OnlineServicesRegistry.h"
+
+UOnlineServicesEngineInterfaceImpl::UOnlineServicesEngineInterfaceImpl(const FObjectInitializer& ObjectInitializer)
+	: Super(ObjectInitializer)
+{
+}
+
+bool UOnlineServicesEngineInterfaceImpl::IsLoaded(FName OnlineIdentifier)
+{
+	// TODO: something like OnlineServicesTypeFromName(OnlineIdentifier); ? This seems to only be called with NAME_None
+	UE::Online::EOnlineServices OnlineServicesType = UE::Online::EOnlineServices::Epic; // TODO:  Need Default support UE::Online::EOnlineServices::Default
+	return UE::Online::IsLoaded(OnlineServicesType, OnlineIdentifier);
+}
+
+FName UOnlineServicesEngineInterfaceImpl::GetOnlineIdentifier(FWorldContext& WorldContext)
+{
+	UE::Online::IOnlineServicesEngineUtils* Utils = UE::Online::GetServicesEngineUtils();
+	if (Utils)
+	{
+		return Utils->GetOnlineIdentifier(WorldContext);
+	}
+	return NAME_None;
+}
+
+FName UOnlineServicesEngineInterfaceImpl::GetOnlineIdentifier(UWorld* World) const
+{
+	UE::Online::IOnlineServicesEngineUtils* Utils = UE::Online::GetServicesEngineUtils();
+	if (Utils)
+	{
+		return Utils->GetOnlineIdentifier(World);
+	}
+	return NAME_None;
+}
+
+bool UOnlineServicesEngineInterfaceImpl::DoesInstanceExist(FName OnlineIdentifier)
+{
+	// TODO:  Does this need to support multiple online service types?  Other accessors seem to not differentiate and just use the default
+	return UE::Online::IsLoaded(UE::Online::EOnlineServices::Default, OnlineIdentifier);
+}
+
+void UOnlineServicesEngineInterfaceImpl::ShutdownOnlineSubsystem(FName OnlineIdentifier)
+{
+	// TODO:  Does this need to support multiple online service types?  Other accessors seem to not differentiate and just use the default
+	if (UE::Online::IOnlineServicesPtr OnlineServices = UE::Online::GetServices(UE::Online::EOnlineServices::Default, OnlineIdentifier))
+	{
+		// TODO:  Is this correct?  Or a new shutdown method?
+		OnlineServices->Destroy();
+	}
+}
+
+void UOnlineServicesEngineInterfaceImpl::DestroyOnlineSubsystem(FName OnlineIdentifier)
+{
+	// TODO:  Does this need to support multiple online service types?  Other accessors seem to not differentiate and just use the default
+	UE::Online::DestroyServices(UE::Online::EOnlineServices::Default, OnlineIdentifier);
+}
+
+FName UOnlineServicesEngineInterfaceImpl::GetDefaultOnlineSubsystemName() const
+{
+	// TODO: Does this have any meaning?  Does this need to be replaced with some kind of FName/EOnlineServices wrapper?
+	return NAME_None;
+}
+
+bool UOnlineServicesEngineInterfaceImpl::IsCompatibleUniqueNetId(const FUniqueNetIdWrapper& InUniqueNetId) const
+{
+	// TODO: Compatibility rules
+	return InUniqueNetId.IsV2();
+}
+
+uint8 UOnlineServicesEngineInterfaceImpl::GetReplicationHashForSubsystem(FName InSubsystemName) const
+{
+	// This is used for net id replication, which doesn't need to call this for V2 ids.
+	checkNoEntry();
+	return 0;
+}
+
+FName UOnlineServicesEngineInterfaceImpl::GetSubsystemFromReplicationHash(uint8 InHash) const
+{
+	// This is used for net id replication, which doesn't need to call this for V2 ids.
+	checkNoEntry();
+	return NAME_None;
+}
+
+FUniqueNetIdWrapper UOnlineServicesEngineInterfaceImpl::CreateUniquePlayerIdWrapper(const FString& Str, FName Type)
+{
+	// This should not be called for OSSv2
+	checkNoEntry();
+	return FUniqueNetIdWrapper();
+}
+
+FUniqueNetIdWrapper UOnlineServicesEngineInterfaceImpl::GetUniquePlayerIdWrapper(UWorld* World, int32 LocalUserNum, FName Type)
+{
+	// TODO: something like OnlineServicesTypeFromName(Type); ? This seems to only be called with NAME_None in LocalPlayer and with the platform service in OnlineReplStructs in a test command
+	FName OnlineIdentifier = GetOnlineIdentifier(World);
+	UE::Online::EOnlineServices OnlineServicesType = UE::Online::EOnlineServices::Epic; // TODO:  Need Default support UE::Online::EOnlineServices::Default
+	if (UE::Online::IOnlineServicesPtr OnlineServices = UE::Online::GetServices(OnlineServicesType, OnlineIdentifier))
+	{
+		if (UE::Online::IAuthPtr AuthPtr = OnlineServices->GetAuthInterface())
+		{
+			UE::Online::FAuthGetAccountByPlatformUserId::Params GetAccountParams = { FPlatformMisc::GetPlatformUserForUserIndex(LocalUserNum) };
+			UE::Online::TOnlineResult<UE::Online::FAuthGetAccountByPlatformUserId> GetAccountResult = AuthPtr->GetAccountByPlatformUserId(MoveTemp(GetAccountParams));
+			if (GetAccountResult.IsOk())
+			{
+				return FUniqueNetIdWrapper(GetAccountResult.GetOkValue().AccountInfo->UserId);
+			}
+		}
+	}
+	return FUniqueNetIdWrapper();
+}
+
+FString UOnlineServicesEngineInterfaceImpl::GetPlayerNickname(UWorld* World, const FUniqueNetIdWrapper& UniqueId)
+{
+	check(UniqueId.IsValid() && UniqueId.IsV2());
+	FName OnlineIdentifier = GetOnlineIdentifier(World);
+	UE::Online::EOnlineServices OnlineServicesType = UE::Online::EOnlineServices::Epic; // TODO: Get from UniqueId
+	if (UE::Online::IOnlineServicesPtr OnlineServices = UE::Online::GetServices(OnlineServicesType, OnlineIdentifier))
+	{
+		if (UE::Online::IAuthPtr AuthPtr = OnlineServices->GetAuthInterface())
+		{
+			UE::Online::FAuthGetAccountByAccountId::Params GetAccountParams = { /* UniqueId.GetV2() */ };
+			UE::Online::TOnlineResult<UE::Online::FAuthGetAccountByAccountId> GetAccountResult = AuthPtr->GetAccountByAccountId(MoveTemp(GetAccountParams));
+			if (GetAccountResult.IsOk())
+			{
+				return GetAccountResult.GetOkValue().AccountInfo->DisplayName;
+			}
+		}
+	}
+	static FString InvalidName(TEXT("InvalidOSSUser"));
+	return InvalidName;
+}
+
+bool UOnlineServicesEngineInterfaceImpl::GetPlayerPlatformNickname(UWorld* World, int32 LocalUserNum, FString& OutNickname)
+{
+	FName OnlineIdentifier = GetOnlineIdentifier(World);
+	UE::Online::EOnlineServices OnlineServicesType = UE::Online::EOnlineServices::Platform;
+	if (UE::Online::IOnlineServicesPtr OnlineServices = UE::Online::GetServices(OnlineServicesType, OnlineIdentifier))
+	{
+		if (UE::Online::IAuthPtr AuthPtr = OnlineServices->GetAuthInterface())
+		{
+			UE::Online::FAuthGetAccountByPlatformUserId::Params GetAccountParams = { FPlatformMisc::GetPlatformUserForUserIndex(LocalUserNum) };
+			UE::Online::TOnlineResult<UE::Online::FAuthGetAccountByPlatformUserId> GetAccountResult = AuthPtr->GetAccountByPlatformUserId(MoveTemp(GetAccountParams));
+			if (GetAccountResult.IsOk())
+			{
+				OutNickname = GetAccountResult.GetOkValue().AccountInfo->DisplayName;
+				return !OutNickname.IsEmpty();
+			}
+		}
+	}
+	return false;
+}
+
+bool UOnlineServicesEngineInterfaceImpl::AutoLogin(UWorld* World, int32 LocalUserNum, const FOnlineAutoLoginComplete& InCompletionDelegate)
+{
+	FName OnlineIdentifier = GetOnlineIdentifier(World);
+	UE::Online::EOnlineServices OnlineServicesType = UE::Online::EOnlineServices::Epic; // TODO:  Need Default support UE::Online::EOnlineServices::Default
+	if (UE::Online::IOnlineServicesPtr OnlineServices = UE::Online::GetServices(OnlineServicesType, OnlineIdentifier))
+	{
+		if (UE::Online::IAuthPtr AuthPtr = OnlineServices->GetAuthInterface())
+		{
+			UE::Online::FAuthLogin::Params LoginParameters;
+			LoginParameters.PlatformUserId = FPlatformMisc::GetPlatformUserForUserIndex(LocalUserNum);
+			// Leave other LoginParameters as default to allow the online service to determine how to try to automatically log in the user
+			UE::Online::TOnlineAsyncOpHandle<UE::Online::FAuthLogin> LoginHandle = AuthPtr->Login(MoveTemp(LoginParameters));
+			LoginHandle.OnComplete(this, [LocalUserNum, InCompletionDelegate](const UE::Online::TOnlineResult<UE::Online::FAuthLogin>& Result)
+			{
+				FString ErrorCode = Result.IsError() ? Result.GetErrorValue().GetLogString() : FString();
+				InCompletionDelegate.ExecuteIfBound(LocalUserNum, Result.IsOk(), ErrorCode);
+			});
+			return true;
+		}
+	}
+	// Not waiting for async login
+	return false;
+}
+
+bool UOnlineServicesEngineInterfaceImpl::IsLoggedIn(UWorld* World, int32 LocalUserNum)
+{
+	FName OnlineIdentifier = GetOnlineIdentifier(World);
+	UE::Online::EOnlineServices OnlineServicesType = UE::Online::EOnlineServices::Epic; // TODO:  Need Default support UE::Online::EOnlineServices::Default
+	if (UE::Online::IOnlineServicesPtr OnlineServices = UE::Online::GetServices(OnlineServicesType, OnlineIdentifier))
+	{
+		if (UE::Online::IAuthPtr AuthPtr = OnlineServices->GetAuthInterface())
+		{
+			UE::Online::FAuthGetAccountByPlatformUserId::Params GetAccountParams = { FPlatformMisc::GetPlatformUserForUserIndex(LocalUserNum) };
+			UE::Online::TOnlineResult<UE::Online::FAuthGetAccountByPlatformUserId> GetAccountResult = AuthPtr->GetAccountByPlatformUserId(MoveTemp(GetAccountParams));
+			if (GetAccountResult.IsOk())
+			{
+				return GetAccountResult.GetOkValue().AccountInfo->LoginStatus == UE::Online::ELoginStatus::LoggedIn;
+			}
+		}
+	}
+	return false;
+}
+
+void UOnlineServicesEngineInterfaceImpl::StartSession(UWorld* World, FName SessionName, FOnlineSessionStartComplete& InCompletionDelegate)
+{
+	// TODO: Implement when we have session/lobby support
+	InCompletionDelegate.ExecuteIfBound(SessionName, false);
+}
+
+void UOnlineServicesEngineInterfaceImpl::EndSession(UWorld* World, FName SessionName, FOnlineSessionEndComplete& InCompletionDelegate)
+{
+	// TODO: Implement when we have session/lobby support
+	InCompletionDelegate.ExecuteIfBound(SessionName, false);
+}
+
+bool UOnlineServicesEngineInterfaceImpl::DoesSessionExist(UWorld* World, FName SessionName)
+{
+	// TODO: Implement when we have session/lobby support
+	return false;
+}
+
+bool UOnlineServicesEngineInterfaceImpl::GetSessionJoinability(UWorld* World, FName SessionName, FJoinabilitySettings& OutSettings)
+{
+	// TODO: Implement when we have session/lobby support
+	return false;
+}
+
+void UOnlineServicesEngineInterfaceImpl::UpdateSessionJoinability(UWorld* World, FName SessionName, bool bPublicSearchable, bool bAllowInvites, bool bJoinViaPresence, bool bJoinViaPresenceFriendsOnly)
+{
+	// TODO: Implement when we have session/lobby support
+}
+
+void UOnlineServicesEngineInterfaceImpl::RegisterPlayer(UWorld* World, FName SessionName, const FUniqueNetIdWrapper& UniqueId, bool bWasInvited)
+{
+	// TODO: Implement when we have session/lobby support
+	check(UniqueId.IsValid() && UniqueId.IsV2());
+}
+
+void UOnlineServicesEngineInterfaceImpl::UnregisterPlayer(UWorld* World, FName SessionName, const FUniqueNetIdWrapper& UniqueId)
+{
+	// TODO: Implement when we have session/lobby support
+	check(UniqueId.IsValid() && UniqueId.IsV2());
+}
+
+void UOnlineServicesEngineInterfaceImpl::UnregisterPlayers(UWorld* World, FName SessionName, const TArray<FUniqueNetIdWrapper>& Players)
+{
+	// TODO: Implement when we have session/lobby support
+	for (const FUniqueNetIdWrapper& PlayerId : Players)
+	{
+		check(PlayerId.IsValid() && PlayerId.IsV2());
+	}
+}
+
+bool UOnlineServicesEngineInterfaceImpl::GetResolvedConnectString(UWorld* World, FName SessionName, FString& URL)
+{
+	// TODO: Implement when we have session/lobby support
+	return false;
+}
+
+TSharedPtr<FVoicePacket> UOnlineServicesEngineInterfaceImpl::GetLocalPacket(UWorld* World, uint8 LocalUserNum)
+{
+	// TODO: Implement when we have voice support
+	return nullptr;
+}
+
+TSharedPtr<FVoicePacket> UOnlineServicesEngineInterfaceImpl::SerializeRemotePacket(UWorld* World, const UNetConnection* const RemoteConnection, FArchive& Ar)
+{
+	// TODO: Implement when we have voice support
+	return nullptr;
+}
+
+void UOnlineServicesEngineInterfaceImpl::StartNetworkedVoice(UWorld* World, uint8 LocalUserNum)
+{
+	// TODO: Implement when we have voice support
+}
+
+void UOnlineServicesEngineInterfaceImpl::StopNetworkedVoice(UWorld* World, uint8 LocalUserNum)
+{
+	// TODO: Implement when we have voice support
+}
+
+void UOnlineServicesEngineInterfaceImpl::ClearVoicePackets(UWorld* World)
+{
+	// TODO: Implement when we have voice support
+}
+
+bool UOnlineServicesEngineInterfaceImpl::MuteRemoteTalker(UWorld* World, uint8 LocalUserNum, const FUniqueNetIdWrapper& PlayerId, bool bIsSystemWide)
+{
+	// TODO: Implement when we have voice support
+	check(PlayerId.IsValid() && PlayerId.IsV2());
+	return false;
+}
+
+bool UOnlineServicesEngineInterfaceImpl::UnmuteRemoteTalker(UWorld* World, uint8 LocalUserNum, const FUniqueNetIdWrapper& PlayerId, bool bIsSystemWide)
+{
+	// TODO: Implement when we have voice support
+	check(PlayerId.IsValid() && PlayerId.IsV2());
+	return false;
+}
+
+int32 UOnlineServicesEngineInterfaceImpl::GetNumLocalTalkers(UWorld* World)
+{
+	// TODO: Implement when we have voice support
+	return 0;
+}
+
+void UOnlineServicesEngineInterfaceImpl::ShowLeaderboardUI(UWorld* World, const FString& CategoryName)
+{
+	FName OnlineIdentifier = GetOnlineIdentifier(World);
+	UE::Online::EOnlineServices OnlineServicesType = UE::Online::EOnlineServices::Epic; // TODO:  Need Default support UE::Online::EOnlineServices::Default
+	if (UE::Online::IOnlineServicesPtr OnlineServices = UE::Online::GetServices(OnlineServicesType, OnlineIdentifier))
+	{
+		UE::Online::IExternalUIPtr ExternalUI = OnlineServices->GetExternalUIInterface();
+		if (ExternalUI.IsValid())
+		{
+			// TODO: When we have more external UI support
+		}
+	}
+}
+
+void UOnlineServicesEngineInterfaceImpl::ShowAchievementsUI(UWorld* World, int32 LocalUserNum)
+{
+	TArray<TSharedRef<UE::Online::IOnlineServices>> ServicesInstances;
+	UE::Online::FOnlineServicesRegistry::Get().GetAllServicesInstances(ServicesInstances);
+	for (const TSharedRef<UE::Online::IOnlineServices>& OnlineServices : ServicesInstances)
+	{
+		UE::Online::IExternalUIPtr ExternalUI = OnlineServices->GetExternalUIInterface();
+		if (ExternalUI.IsValid())
+		{
+			// TODO: When we have more external UI support
+		}
+	}
+}
+
+void UOnlineServicesEngineInterfaceImpl::ShowWebURL(const FString& CurrentURL, const UOnlineEngineInterface::FShowWebUrlParams& ShowParams, const FOnlineShowWebUrlClosed& CompletionDelegate)
+{
+	TArray<TSharedRef<UE::Online::IOnlineServices>> ServicesInstances;
+	UE::Online::FOnlineServicesRegistry::Get().GetAllServicesInstances(ServicesInstances);
+	for (const TSharedRef<UE::Online::IOnlineServices>& OnlineServices : ServicesInstances)
+	{
+		UE::Online::IExternalUIPtr ExternalUI = OnlineServices->GetExternalUIInterface();
+		if (ExternalUI.IsValid())
+		{
+			// TODO: When we have more external UI support
+		}
+	}
+}
+
+bool UOnlineServicesEngineInterfaceImpl::CloseWebURL()
+{
+	TArray<TSharedRef<UE::Online::IOnlineServices>> ServicesInstances;
+	UE::Online::FOnlineServicesRegistry::Get().GetAllServicesInstances(ServicesInstances);
+	for (const TSharedRef<UE::Online::IOnlineServices>& OnlineServices : ServicesInstances)
+	{
+		UE::Online::IExternalUIPtr ExternalUI = OnlineServices->GetExternalUIInterface();
+		if (ExternalUI.IsValid())
+		{
+			// TODO: When we have more external UI support
+		}
+	}
+	return false;
+}
+
+void UOnlineServicesEngineInterfaceImpl::BindToExternalUIOpening(const FOnlineExternalUIChanged& Delegate)
+{
+	UE::Online::IOnlineServicesEngineUtils* Utils = UE::Online::GetServicesEngineUtils();
+	if (Utils)
+	{
+		FOnExternalUIChangeDelegate OnExternalUIChangeDelegate;
+		OnExternalUIChangeDelegate.BindWeakLambda(this, [Delegate](bool bInIsOpening)
+		{
+			Delegate.ExecuteIfBound(bInIsOpening);
+		});
+		Utils->SetEngineExternalUIBinding(OnExternalUIChangeDelegate);
+	}
+}
+
+void UOnlineServicesEngineInterfaceImpl::DumpSessionState(UWorld* World)
+{
+	// TODO: Implement when we have session/lobby support
+}
+
+void UOnlineServicesEngineInterfaceImpl::DumpPartyState(UWorld* World)
+{
+	// TODO: Implement when we have lobby support
+}
+
+void UOnlineServicesEngineInterfaceImpl::DumpVoiceState(UWorld* World)
+{
+	// TODO: Implement when we have voice support
+}
+
+void UOnlineServicesEngineInterfaceImpl::DumpChatState(UWorld* World)
+{
+	// TODO: Implement when we have chat support
+}
+
+#if WITH_EDITOR
+bool UOnlineServicesEngineInterfaceImpl::SupportsOnlinePIE()
+{
+	return UE::Online::GetServicesEngineUtils()->SupportsOnlinePIE();
+}
+
+void UOnlineServicesEngineInterfaceImpl::SetShouldTryOnlinePIE(bool bShouldTry)
+{
+	UE::Online::GetServicesEngineUtils()->SetShouldTryOnlinePIE(bShouldTry);
+}
+
+int32 UOnlineServicesEngineInterfaceImpl::GetNumPIELogins()
+{
+	return UE::Online::GetServicesEngineUtils()->GetNumPIELogins();
+}
+
+void UOnlineServicesEngineInterfaceImpl::SetForceDedicated(FName OnlineIdentifier, bool bForce)
+{
+	// TODO:  Support other services? This is only called from PlayLevel with the result of GetOnlineIdentifier which only returns the default service id
+	UE::Online::EOnlineServices OnlineServicesType = UE::Online::EOnlineServices::Epic; // TODO:  Need Default support UE::Online::EOnlineServices::Default
+	if (UE::Online::IOnlineServicesPtr OnlineServices = UE::Online::GetServices(OnlineServicesType, OnlineIdentifier))
+	{
+		// TODO:  Add SetForceDedicated method to OnlineServices
+	}
+}
+
+void UOnlineServicesEngineInterfaceImpl::LoginPIEInstance(FName OnlineIdentifier, int32 LocalUserNum, int32 PIELoginNum, FOnPIELoginComplete& CompletionDelegate)
+{
+	FString ErrorStr;
+	if (SupportsOnlinePIE())
+	{
+		TArray<UE::Online::FAuthLogin::Params> PIELogins;
+		UE::Online::GetServicesEngineUtils()->GetPIELogins(PIELogins);
+		if (PIELogins.IsValidIndex(PIELoginNum))
+		{
+			// TODO:  Support other services? This is only called from PlayLevel with the result of GetOnlineIdentifier which only returns the default service id
+			UE::Online::EOnlineServices OnlineServicesType = UE::Online::EOnlineServices::Epic; // TODO:  Need Default support UE::Online::EOnlineServices::Default
+			if (UE::Online::IOnlineServicesPtr OnlineServices = UE::Online::GetServices(OnlineServicesType, OnlineIdentifier))
+			{
+				if (UE::Online::IAuthPtr AuthPtr = OnlineServices->GetAuthInterface())
+				{
+					UE::Online::FAuthLogin::Params LoginParameters = PIELogins[PIELoginNum];
+					LoginParameters.PlatformUserId = FPlatformMisc::GetPlatformUserForUserIndex(LocalUserNum);
+					AuthPtr->Login(MoveTemp(LoginParameters)).OnComplete(this, [LocalUserNum, CompletionDelegate](const UE::Online::TOnlineResult<UE::Online::FAuthLogin>& Result)
+					{
+						FString AuthLoginErrorStr;
+						if (Result.IsError())
+						{
+							AuthLoginErrorStr = Result.GetErrorValue().GetLogString();
+						}
+						CompletionDelegate.ExecuteIfBound(LocalUserNum, Result.IsOk(), AuthLoginErrorStr);
+					});
+					return;
+				}
+				else
+				{
+					ErrorStr = TEXT("No auth interface to login");
+				}
+			}
+			else
+			{
+				ErrorStr = TEXT("No online service to login");
+			}
+		}
+		else
+		{
+			ErrorStr = FString::Printf(TEXT("Invalid credential index for PIE login. Index: %d NumLogins: %d"), PIELoginNum, PIELogins.Num());
+		}
+	}
+	else
+	{
+		ErrorStr = TEXT("PIE login not supported");
+	}
+
+	if (!ErrorStr.IsEmpty())
+	{
+		CompletionDelegate.ExecuteIfBound(LocalUserNum, false, ErrorStr);
+	}
+}
+
+#endif
diff --git a/Engine/Plugins/Online/OnlineSubsystemUtils/Source/OnlineSubsystemUtils/Private/OnlineServicesEngineInterfaceImpl.h b/Engine/Plugins/Online/OnlineSubsystemUtils/Source/OnlineSubsystemUtils/Private/OnlineServicesEngineInterfaceImpl.h
new file mode 100644
--- /dev/null
+++ b/Engine/Plugins/Online/OnlineSubsystemUtils/Source/OnlineSubsystemUtils/Private/OnlineServicesEngineInterfaceImpl.h
@@ -0,0 +1,129 @@
+// Copyright Epic Games, Inc. All Rights Reserved.
+
+#pragma once
+
+#include "CoreMinimal.h"
+#include "UObject/ObjectMacros.h"
+#include "Online/CoreOnline.h"
+#include "Net/OnlineEngineInterface.h"
+#include "OnlineServicesEngineInterfaceImpl.generated.h"
+
+class Error;
+class FVoicePacket;
+struct FWorldContext;
+
+/**
+ * Implementation of UOnlineEngineInterface that uses Online Services (also known as Online Subsystem v2)
+ */
+UCLASS(config = Engine)
+class ONLINESUBSYSTEMUTILS_API UOnlineServicesEngineInterfaceImpl : public UOnlineEngineInterface
+{
+	GENERATED_UCLASS_BODY()
+
+	/**
+	 * Subsystem
+	 */
+public:
+
+	virtual bool IsLoaded(FName OnlineIdentifier) override;
+	virtual FName GetOnlineIdentifier(FWorldContext& WorldContext) override;
+	virtual bool DoesInstanceExist(FName OnlineIdentifier) override;
+	virtual void ShutdownOnlineSubsystem(FName OnlineIdentifier) override;
+	virtual void DestroyOnlineSubsystem(FName OnlineIdentifier) override;
+	virtual FName GetDefaultOnlineSubsystemName() const override;
+	virtual bool IsCompatibleUniqueNetId(const FUniqueNetIdWrapper& InUniqueNetId) const override;
+
+	/**
+	 * Utils
+	 */
+	virtual uint8 GetReplicationHashForSubsystem(FName InSubsystemName) const override;
+	virtual FName GetSubsystemFromReplicationHash(uint8 InHash) const override;
+
+private:
+	/** @return the identifier/context handle associated with a UWorld */
+	FName GetOnlineIdentifier(UWorld* World) const;
+
+	/**
+	 * Identity
+	 */
+public:
+
+	virtual FUniqueNetIdWrapper CreateUniquePlayerIdWrapper(const FString& Str, FName Type) override;
+	virtual FUniqueNetIdWrapper GetUniquePlayerIdWrapper(UWorld* World, int32 LocalUserNum, FName Type) override;
+
+	virtual FString GetPlayerNickname(UWorld* World, const FUniqueNetIdWrapper& UniqueId) override;
+	virtual bool GetPlayerPlatformNickname(UWorld* World, int32 LocalUserNum, FString& OutNickname) override;
+
+	virtual bool AutoLogin(UWorld* World, int32 LocalUserNum, const FOnlineAutoLoginComplete& InCompletionDelegate) override;
+	virtual bool IsLoggedIn(UWorld* World, int32 LocalUserNum) override;
+
+	/**
+	 * Session
+	 */
+public:
+
+	virtual void StartSession(UWorld* World, FName SessionName, FOnlineSessionStartComplete& InCompletionDelegate) override;
+	virtual void EndSession(UWorld* World, FName SessionName, FOnlineSessionEndComplete& InCompletionDelegate) override;
+	virtual bool DoesSessionExist(UWorld* World, FName SessionName) override;
+
+	virtual bool GetSessionJoinability(UWorld* World, FName SessionName, FJoinabilitySettings& OutSettings) override;
+	virtual void UpdateSessionJoinability(UWorld* World, FName SessionName, bool bPublicSearchable, bool bAllowInvites, bool bJoinViaPresence, bool bJoinViaPresenceFriendsOnly) override;
+
+	virtual void RegisterPlayer(UWorld* World, FName SessionName, const FUniqueNetIdWrapper& UniqueId, bool bWasInvited) override;
+	virtual void UnregisterPlayer(UWorld* World, FName SessionName, const FUniqueNetIdWrapper& UniqueId) override;
+	virtual void UnregisterPlayers(UWorld* World, FName SessionName, const TArray<FUniqueNetIdWrapper>& Players) override;
+
+	virtual bool GetResolvedConnectString(UWorld* World, FName SessionName, FString& URL) override;
+
+	/**
+	 * Voice
+	 */
+public:
+
+	virtual TSharedPtr<FVoicePacket> GetLocalPacket(UWorld* World, uint8 LocalUserNum) override;
+	virtual TSharedPtr<FVoicePacket> SerializeRemotePacket(UWorld* World, const UNetConnection* const RemoteConnection, FArchive& Ar) override;
+
+	virtual void StartNetworkedVoice(UWorld* World, uint8 LocalUserNum) override;
+	virtual void StopNetworkedVoice(UWorld* World, uint8 LocalUserNum) override;
+	virtual void ClearVoicePackets(UWorld* World) override;
+
+	virtual bool MuteRemoteTalker(UWorld* World, uint8 LocalUserNum, const FUniqueNetIdWrapper& PlayerId, bool bIsSystemWide) override;
+	virtual bool UnmuteRemoteTalker(UWorld* World, uint8 LocalUserNum, const FUniqueNetIdWrapper& PlayerId, bool bIsSystemWide) override;
+
+	virtual int32 GetNumLocalTalkers(UWorld* World) override;
+
+	/**
+	 * External UI
+	 */
+public:
+
+	virtual void ShowLeaderboardUI(UWorld* World, const FString& CategoryName) override;
+	virtual void ShowAchievementsUI(UWorld* World, int32 LocalUserNum) override;
+	virtual void BindToExternalUIOpening(const FOnlineExternalUIChanged& Delegate) override;
+	virtual void ShowWebURL(const FString& CurrentURL, const UOnlineEngineInterface::FShowWebUrlParams& ShowParams, const FOnlineShowWebUrlClosed& CompletionDelegate) override;
+	virtual bool CloseWebURL() override;
+
+	/**
+	 * Debug
+	 */
+public:
+
+	virtual void DumpSessionState(UWorld* World) override;
+	virtual void DumpPartyState(UWorld* World) override;
+	virtual void DumpVoiceState(UWorld* World) override;
+	virtual void DumpChatState(UWorld* World) override;
+
+#if WITH_EDITOR
+	/**
+	 * PIE Utilities
+	 */
+public:
+
+	virtual bool SupportsOnlinePIE() override;
+	virtual void SetShouldTryOnlinePIE(bool bShouldTry) override;
+	virtual int32 GetNumPIELogins() override;
+	virtual void SetForceDedicated(FName OnlineIdentifier, bool bForce) override;
+	virtual void LoginPIEInstance(FName OnlineIdentifier, int32 LocalUserNum, int32 PIELoginNum, FOnPIELoginComplete& CompletionDelegate) override;
+#endif
+};
+
diff --git a/Engine/Plugins/Online/OnlineSubsystemUtils/Source/OnlineSubsystemUtils/Private/OnlineServicesEngineUtilsImpl.cpp b/Engine/Plugins/Online/OnlineSubsystemUtils/Source/OnlineSubsystemUtils/Private/OnlineServicesEngineUtilsImpl.cpp
new file mode 100644
--- /dev/null
+++ b/Engine/Plugins/Online/OnlineSubsystemUtils/Source/OnlineSubsystemUtils/Private/OnlineServicesEngineUtilsImpl.cpp
@@ -0,0 +1,164 @@
+// Copyright Epic Games, Inc. All Rights Reserved.
+
+#include "OnlineServicesEngineUtilsImpl.h"
+#include "OnlineDelegates.h"
+#include "OnlinePIESettings.h"
+#include "Modules/ModuleManager.h"
+#include "Engine/Engine.h"
+#include "OnlineSubsystemUtilsModule.h"
+
+#include "Online/CoreOnline.h"
+#include "Online/ExternalUI.h"
+#include "Online/OnlineServices.h"
+#include "Online/OnlineServicesDelegates.h"
+#include "Online/OnlineServicesRegistry.h"
+
+namespace UE::Online
+{
+
+IOnlineServicesEngineUtils* GetServicesEngineUtils()
+{
+	static const FName OnlineSubsystemModuleName = TEXT("OnlineSubsystemUtils");
+	FOnlineSubsystemUtilsModule* OSSUtilsModule = FModuleManager::GetModulePtr<FOnlineSubsystemUtilsModule>(OnlineSubsystemModuleName);
+	if (OSSUtilsModule != nullptr)
+	{
+		return OSSUtilsModule->GetServicesUtils();
+	}
+
+	return nullptr;
+}
+
+FOnlineServicesEngineUtils::~FOnlineServicesEngineUtils()
+{
+	UE::Online::OnOnlineServicesCreated.Remove(OnOnlineServicesCreatedDelegateHandle);
+}
+
+FName FOnlineServicesEngineUtils::GetOnlineIdentifier(const FWorldContext& WorldContext) const
+{
+#if WITH_EDITOR
+	if (WorldContext.WorldType == EWorldType::PIE)
+	{
+		return WorldContext.ContextHandle;
+	}
+#endif
+	return NAME_None;
+}
+
+FName FOnlineServicesEngineUtils::GetOnlineIdentifier(const UWorld* World) const
+{
+#if WITH_EDITOR
+	if (const FWorldContext* WorldContext = GEngine->GetWorldContextFromWorld(World))
+	{
+		return GetOnlineIdentifier(*WorldContext);
+	}
+#endif
+	return NAME_None;
+}
+
+void FOnlineServicesEngineUtils::OnOnlineServicesCreated(TSharedRef<class IOnlineServices> NewServices)
+{
+	if (OnExternalUIChangeDelegate.IsBound())
+	{
+		IExternalUIPtr ExternalUI = NewServices->GetExternalUIInterface();
+		if (ExternalUI.IsValid())
+		{
+			ExternalUIDelegateHandles.Emplace(ExternalUI->OnExternalUIStatusChanged().Add([this](const FExternalUIStatusChanged& EventParams)
+			{
+				OnExternalUIChangeDelegate.ExecuteIfBound(EventParams.bIsOpening);
+			}));
+		}
+	}
+}
+
+void FOnlineServicesEngineUtils::SetEngineExternalUIBinding(const FOnExternalUIChangeDelegate& InOnExternalUIChangeDelegate)
+{
+	OnExternalUIChangeDelegate = InOnExternalUIChangeDelegate;
+	TArray<TSharedRef<IOnlineServices>> ServicesInstances;
+	FOnlineServicesRegistry::Get().GetAllServicesInstances(ServicesInstances);
+	for (const TSharedRef<IOnlineServices>& ServiceInstance : ServicesInstances)
+	{
+		if (IExternalUIPtr ExternalUI = ServiceInstance->GetExternalUIInterface())
+		{
+			ExternalUIDelegateHandles.Emplace(ExternalUI->OnExternalUIStatusChanged().Add([this](const FExternalUIStatusChanged& EventParams)
+			{
+				OnExternalUIChangeDelegate.ExecuteIfBound(EventParams.bIsOpening);
+			}));
+		}
+	};
+}
+
+#if WITH_EDITOR
+bool FOnlineServicesEngineUtils::SupportsOnlinePIE() const
+{
+	check(UObjectInitialized());
+	const UOnlinePIESettings* OnlinePIESettings = GetDefault<UOnlinePIESettings>();
+	if (OnlinePIESettings->bOnlinePIEEnabled && GetNumPIELogins() > 0)
+	{
+		// If we can't get the auth then things are either not configured right or disabled
+		EOnlineServices OnlineServicesType = EOnlineServices::Epic; // TODO:  Need Default support UE::Online::EOnlineServices::Default
+		if (IOnlineServicesPtr OnlineServices = GetServices(OnlineServicesType))
+		{
+			if (IAuthPtr AuthPtr = OnlineServices->GetAuthInterface())
+			{
+				return true;
+			}
+		}
+	}
+	return false;
+}
+
+void FOnlineServicesEngineUtils::SetShouldTryOnlinePIE(bool bShouldTry)
+{
+	if (bShouldTryOnlinePIE != bShouldTry)
+	{
+		bShouldTryOnlinePIE = bShouldTry;
+		// TODO:  Reload online services?
+	}
+}
+
+bool FOnlineServicesEngineUtils::IsOnlinePIEEnabled() const
+{
+	check(UObjectInitialized());
+	const UOnlinePIESettings* OnlinePIESettings = GetDefault<UOnlinePIESettings>();
+	return bShouldTryOnlinePIE && OnlinePIESettings->bOnlinePIEEnabled;
+}
+
+int32 FOnlineServicesEngineUtils::GetNumPIELogins() const
+{
+	check(UObjectInitialized());
+
+	int32 NumValidLogins = 0;
+	const UOnlinePIESettings* OnlinePIESettings = GetDefault<UOnlinePIESettings>();
+	for (const FPIELoginSettingsInternal& Login : OnlinePIESettings->Logins)
+	{
+		if (Login.IsValid())
+		{
+			NumValidLogins++;
+		}
+	}
+	
+	return NumValidLogins;
+}
+
+void FOnlineServicesEngineUtils::GetPIELogins(TArray<FAuthLogin::Params>& Logins) const
+{
+	check(UObjectInitialized());
+	const UOnlinePIESettings* OnlinePIESettings = GetDefault<UOnlinePIESettings>();
+	for (const FPIELoginSettingsInternal& PIELoginCredentials : OnlinePIESettings->Logins)
+	{
+		UE::Online::FAuthLogin::Params LoginParameters;
+		LoginParameters.CredentialsId = PIELoginCredentials.Id;
+		LoginParameters.CredentialsToken.Set<FString>(PIELoginCredentials.Token);
+		LoginParameters.CredentialsType = PIELoginCredentials.Type;
+		Logins.Emplace(MoveTemp(LoginParameters));
+	}
+}
+
+#endif // WITH_EDITOR
+
+void FOnlineServicesEngineUtils::Init()
+{
+	OnOnlineServicesCreatedDelegateHandle = UE::Online::OnOnlineServicesCreated.AddRaw(this, &FOnlineServicesEngineUtils::OnOnlineServicesCreated);
+}
+
+}
diff --git a/Engine/Plugins/Online/OnlineSubsystemUtils/Source/OnlineSubsystemUtils/Private/OnlineServicesEngineUtilsImpl.h b/Engine/Plugins/Online/OnlineSubsystemUtils/Source/OnlineSubsystemUtils/Private/OnlineServicesEngineUtilsImpl.h
new file mode 100644
--- /dev/null
+++ b/Engine/Plugins/Online/OnlineSubsystemUtils/Source/OnlineSubsystemUtils/Private/OnlineServicesEngineUtilsImpl.h
@@ -0,0 +1,61 @@
+// Copyright Epic Games, Inc. All Rights Reserved.
+
+#pragma once
+
+#include "OnlineServicesEngineUtils.h"
+#include "OnlineDelegates.h"
+#include "OnlinePIESettings.h"
+#include "Modules/ModuleManager.h"
+#include "Engine/Engine.h"
+#include "OnlineSubsystemUtilsModule.h"
+
+namespace UE::Online
+{
+
+class IOnlineServices;
+
+/**
+ * Concrete implementation of IOnlineServicesEngineUtils interface 
+ */
+class FOnlineServicesEngineUtils : public IOnlineServicesEngineUtils
+{
+public:
+
+	FOnlineServicesEngineUtils()
+		: bShouldTryOnlinePIE(true)
+	{
+	}
+
+	virtual ~FOnlineServicesEngineUtils();
+
+	virtual FName GetOnlineIdentifier(const FWorldContext& WorldContext) const override;
+	virtual FName GetOnlineIdentifier(const UWorld* World) const override;
+	virtual void SetEngineExternalUIBinding(const FOnExternalUIChangeDelegate& InOnExternalUIChangeDelegate) override;
+
+#if WITH_EDITOR
+	virtual bool SupportsOnlinePIE() const override;
+	virtual void SetShouldTryOnlinePIE(bool bShouldTry) override;
+	virtual bool IsOnlinePIEEnabled() const override;
+	virtual int32 GetNumPIELogins() const override;
+	virtual void GetPIELogins(TArray<FAuthLogin::Params>& Logins) const override;
+#endif // WITH_EDITOR
+
+private:
+
+	void Init();
+	void OnOnlineServicesCreated(TSharedRef<class IOnlineServices> NewServices);
+
+	/** If false it will not try to do online PIE at all */
+	bool bShouldTryOnlinePIE;
+
+	/** Delegate set by the engine for notification of external UI operations */
+	FOnExternalUIChangeDelegate OnExternalUIChangeDelegate;
+	/** Array of delegate handles */
+	TArray<FOnlineEventDelegateHandle> ExternalUIDelegateHandles;
+
+	/** Delegate binding when new online services are created */
+	FDelegateHandle OnOnlineServicesCreatedDelegateHandle;
+	friend FOnlineSubsystemUtilsModule;
+};
+
+}
diff --git a/Engine/Plugins/Online/OnlineSubsystemUtils/Source/OnlineSubsystemUtils/Private/OnlineSubsystemUtilsModule.cpp b/Engine/Plugins/Online/OnlineSubsystemUtils/Source/OnlineSubsystemUtils/Private/OnlineSubsystemUtilsModule.cpp
--- a/Engine/Plugins/Online/OnlineSubsystemUtils/Source/OnlineSubsystemUtils/Private/OnlineSubsystemUtilsModule.cpp
+++ b/Engine/Plugins/Online/OnlineSubsystemUtils/Source/OnlineSubsystemUtils/Private/OnlineSubsystemUtilsModule.cpp
@@ -5,6 +5,7 @@
 #include "Interfaces/OnlineIdentityInterface.h"
 #include "OnlineSubsystem.h"
 #include "OnlineSubsystemUtils.h"
+#include "OnlineServicesEngineUtilsImpl.h"
 #include "OnlineDelegates.h"
 #include "OnlinePIESettings.h"
 
@@ -248,14 +249,27 @@ PRAGMA_ENABLE_DEPRECATION_WARNINGS
 	friend FOnlineSubsystemUtilsModule;
 };
 
+FOnlineSubsystemUtilsModule::FOnlineSubsystemUtilsModule()
+{
+}
+
+FOnlineSubsystemUtilsModule::~FOnlineSubsystemUtilsModule()
+{
+}
+
 void FOnlineSubsystemUtilsModule::StartupModule()
 {
 	TUniquePtr<FOnlineSubsystemUtils> TempUtils = MakeUnique<FOnlineSubsystemUtils>();
 	TempUtils->Init();
 	SubsystemUtils = MoveTemp(TempUtils);
+
+	TUniquePtr<UE::Online::FOnlineServicesEngineUtils> TempServicesUtils = MakeUnique<UE::Online::FOnlineServicesEngineUtils>();
+	TempServicesUtils->Init();
+	OnlineServicesUtils = MoveTemp(TempServicesUtils);
 }
 
 void FOnlineSubsystemUtilsModule::ShutdownModule()
 {
 	SubsystemUtils.Reset();
+	OnlineServicesUtils.Reset();
 }
diff --git a/Engine/Plugins/Online/OnlineSubsystemUtils/Source/OnlineSubsystemUtils/Public/OnlineServicesEngineUtils.h b/Engine/Plugins/Online/OnlineSubsystemUtils/Source/OnlineSubsystemUtils/Public/OnlineServicesEngineUtils.h
new file mode 100644
--- /dev/null
+++ b/Engine/Plugins/Online/OnlineSubsystemUtils/Source/OnlineSubsystemUtils/Public/OnlineServicesEngineUtils.h
@@ -0,0 +1,75 @@
+// Copyright Epic Games, Inc. All Rights Reserved.
+
+#pragma once
+
+#include "CoreMinimal.h"
+#include "Online/CoreOnline.h"
+#include "Online/Auth.h"
+
+class UWorld;
+struct FWorldContext;
+struct FPIELoginSettingsInternal;
+
+DECLARE_MULTICAST_DELEGATE_OneParam(FOnExternalUIChange, bool);
+typedef FOnExternalUIChange::FDelegate FOnExternalUIChangeDelegate;
+
+namespace UE::Online {
+
+/**
+ * Interface class for various online utility functions for Engine
+ */
+class IOnlineServicesEngineUtils
+{
+protected:
+	/** Hidden on purpose */
+	IOnlineServicesEngineUtils() {}
+public:
+
+	virtual ~IOnlineServicesEngineUtils() {}
+
+	/**
+	 * Gets an FName that uniquely identifies an instance of a world to differentiate online services for PIE
+	 *
+	 * @param WorldContext the worldcontext
+	 * @return an FName that can be used with online services
+	 */
+	virtual FName GetOnlineIdentifier(const FWorldContext& WorldContext) const = 0;
+
+	/**
+	 * Gets an FName that uniquely identifies an instance of a world to differentiate online services for PIE
+	 *
+	 * @param World the world to use for context
+	 * @return an FName that can be used with Online Services
+	 */
+	virtual FName GetOnlineIdentifier(const UWorld* World) const = 0;
+
+	/**
+	 * Bind a notification delegate when any subsystem external UI is opened/closed
+	 * *NOTE* there is only meant to be one delegate needed for this, game code should bind manually
+	 *
+	 * @param OnExternalUIChangeDelegate delegate fired when the external UI is opened/closed
+	 */
+	virtual void SetEngineExternalUIBinding(const FOnExternalUIChangeDelegate& OnExternalUIChangeDelegate) = 0;
+
+#if WITH_EDITOR
+	/**
+	 * Play in Editor settings
+	 */
+
+	/** @return true if the default platform supports logging in for Play In Editor (PIE) */
+	virtual bool SupportsOnlinePIE() const = 0;
+	/** Enable/Disable online PIE at runtime */
+	virtual void SetShouldTryOnlinePIE(bool bShouldTry) = 0;
+	/** @return true if the user has enabled logging in for Play In Editor (PIE) */
+	virtual bool IsOnlinePIEEnabled() const = 0;
+	/** @return the number of logins the user has setup for Play In Editor (PIE) */
+	virtual int32 GetNumPIELogins() const = 0;
+	/** @return the array of valid credentials the user has setup for Play In Editor (PIE) */
+	virtual void GetPIELogins(TArray<FAuthLogin::Params>& Logins) const = 0;
+#endif // WITH_EDITOR
+};
+
+/** @return the single instance of the online services utils interface */
+ONLINESUBSYSTEMUTILS_API IOnlineServicesEngineUtils* GetServicesEngineUtils();
+
+/* UE::Online */ }
diff --git a/Engine/Plugins/Online/OnlineSubsystemUtils/Source/OnlineSubsystemUtils/Public/OnlineSubsystemUtilsModule.h b/Engine/Plugins/Online/OnlineSubsystemUtils/Source/OnlineSubsystemUtils/Public/OnlineSubsystemUtilsModule.h
--- a/Engine/Plugins/Online/OnlineSubsystemUtils/Source/OnlineSubsystemUtils/Public/OnlineSubsystemUtilsModule.h
+++ b/Engine/Plugins/Online/OnlineSubsystemUtils/Source/OnlineSubsystemUtils/Public/OnlineSubsystemUtilsModule.h
@@ -6,6 +6,9 @@
 #include "Modules/ModuleInterface.h"
 
 class IOnlineSubsystemUtils;
+namespace UE::Online {
+	class IOnlineServicesEngineUtils;
+}
 
 /**
  * Online subsystem utils module class
@@ -15,12 +18,15 @@ class FOnlineSubsystemUtilsModule : public IModuleInterface
 {
 public:
 
-	FOnlineSubsystemUtilsModule() {}
-	virtual ~FOnlineSubsystemUtilsModule() {}
+	FOnlineSubsystemUtilsModule();
+	virtual ~FOnlineSubsystemUtilsModule();
 
 	/** @return the singleton utility interface */
 	IOnlineSubsystemUtils* GetUtils() const { return SubsystemUtils.Get(); }
 
+	/** @return the singleton utility interface */
+	UE::Online::IOnlineServicesEngineUtils* GetServicesUtils() const { return OnlineServicesUtils.Get(); }
+
 	// IModuleInterface
 	virtual void StartupModule() override;
 	virtual void ShutdownModule() override;
@@ -40,6 +46,9 @@ private:
 
 	/** Online subsystem utilities singleton */
 	TUniquePtr<IOnlineSubsystemUtils> SubsystemUtils;
+
+	/** Online Services (OSSv2) utilities singleton */
+	TUniquePtr<UE::Online::IOnlineServicesEngineUtils> OnlineServicesUtils;
 };
 
 
diff --git a/Engine/Source/Runtime/CoreOnline/Public/Online/CoreOnline.h b/Engine/Source/Runtime/CoreOnline/Public/Online/CoreOnline.h
--- a/Engine/Source/Runtime/CoreOnline/Public/Online/CoreOnline.h
+++ b/Engine/Source/Runtime/CoreOnline/Public/Online/CoreOnline.h
@@ -382,6 +382,11 @@ struct FUniqueNetIdWrapper
 	{
 	}
 
+	FUniqueNetIdWrapper(const UE::Online::FOnlineAccountIdHandle& Handle)
+	{
+		Variant.Emplace<UE::Online::FOnlineAccountIdHandle>(Handle);
+	}
+
 	// temporarily restored implicit conversion from FUniqueNetId
 	FUniqueNetIdWrapper(const FUniqueNetId& InUniqueNetId)
 	{
diff --git a/Engine/Source/Runtime/Engine/Private/GameModeBase.cpp b/Engine/Source/Runtime/Engine/Private/GameModeBase.cpp
--- a/Engine/Source/Runtime/Engine/Private/GameModeBase.cpp
+++ b/Engine/Source/Runtime/Engine/Private/GameModeBase.cpp
@@ -633,7 +633,7 @@ void AGameModeBase::GameWelcomePlayer(UNetConnection* Connection, FString& Redir
 void AGameModeBase::PreLogin(const FString& Options, const FString& Address, const FUniqueNetIdRepl& UniqueId, FString& ErrorMessage)
 {
 	// Login unique id must match server expected unique id type OR No unique id could mean game doesn't use them
-	const bool bUniqueIdCheckOk = (!UniqueId.IsValid() || UOnlineEngineInterface::Get()->IsCompatibleUniqueNetId(*UniqueId));
+	const bool bUniqueIdCheckOk = (!UniqueId.IsValid() || UOnlineEngineInterface::Get()->IsCompatibleUniqueNetId(UniqueId));
 	if (bUniqueIdCheckOk)
 	{
 		ErrorMessage = GameSession->ApproveLogin(Options);
diff --git a/Engine/Source/Runtime/Engine/Private/GameSession.cpp b/Engine/Source/Runtime/Engine/Private/GameSession.cpp
--- a/Engine/Source/Runtime/Engine/Private/GameSession.cpp
+++ b/Engine/Source/Runtime/Engine/Private/GameSession.cpp
@@ -247,18 +247,23 @@ void AGameSession::UnregisterPlayer(FName InSessionName, const FUniqueNetIdRepl&
 		UniqueId->IsValid())
 	{
 		// Remove the player from the session
-		UOnlineEngineInterface::Get()->UnregisterPlayer(World, InSessionName, *UniqueId);
+		UOnlineEngineInterface::Get()->UnregisterPlayer(World, InSessionName, UniqueId);
 	}
 }
 
-void AGameSession::UnregisterPlayers(FName InSessionName, const TArray< FUniqueNetIdRef >& Players)
+void AGameSession::UnregisterPlayers(FName InSessionName, const TArray<FUniqueNetIdRef>& Players)
 {
 	UWorld* World = GetWorld();
 	if (GetNetMode() != NM_Standalone &&
 		Players.Num() > 0)
 	{
 		// Remove the player from the session
-		UOnlineEngineInterface::Get()->UnregisterPlayers(World, InSessionName, Players);
+		TArray<FUniqueNetIdWrapper> PlayerIdsAsWrappers;
+		for (const FUniqueNetIdRef& PlayerId : Players)
+		{
+			PlayerIdsAsWrappers.Emplace(PlayerId);
+		}
+		UOnlineEngineInterface::Get()->UnregisterPlayers(World, InSessionName, PlayerIdsAsWrappers);
 	}
 }
 
diff --git a/Engine/Source/Runtime/Engine/Private/LocalPlayer.cpp b/Engine/Source/Runtime/Engine/Private/LocalPlayer.cpp
--- a/Engine/Source/Runtime/Engine/Private/LocalPlayer.cpp
+++ b/Engine/Source/Runtime/Engine/Private/LocalPlayer.cpp
@@ -1581,10 +1581,10 @@ FString ULocalPlayer::GetNickname() const
 			return PlatformNickname;
 		}
 
-		auto UniqueId = GetPreferredUniqueNetId();
+		FUniqueNetIdRepl UniqueId = GetPreferredUniqueNetId();
 		if (UniqueId.IsValid())
 		{
-			return UOnlineEngineInterface::Get()->GetPlayerNickname(World, *UniqueId);
+			return UOnlineEngineInterface::Get()->GetPlayerNickname(World, UniqueId);
 		}
 	}
 
@@ -1596,7 +1596,7 @@ FUniqueNetIdRepl ULocalPlayer::GetUniqueNetIdFromCachedControllerId() const
 	UWorld* World = GetWorld();
 	if (World != nullptr)
 	{
-		return FUniqueNetIdRepl(UOnlineEngineInterface::Get()->GetUniquePlayerId(World, ControllerId));
+		return FUniqueNetIdRepl(UOnlineEngineInterface::Get()->GetUniquePlayerIdWrapper(World, ControllerId));
 	}
 
 	return FUniqueNetIdRepl();
@@ -1607,7 +1607,7 @@ FUniqueNetIdRepl ULocalPlayer::GetUniqueNetIdForPlatformUser() const
 	UWorld* World = GetWorld();
 	if (World != nullptr)
 	{
-		return FUniqueNetIdRepl(UOnlineEngineInterface::Get()->GetUniquePlayerId(World, PlatformUserId));
+		return FUniqueNetIdRepl(UOnlineEngineInterface::Get()->GetUniquePlayerIdWrapper(World, PlatformUserId));
 	}
 
 	return FUniqueNetIdRepl();
diff --git a/Engine/Source/Runtime/Engine/Private/OnlineEngineInterface.cpp b/Engine/Source/Runtime/Engine/Private/OnlineEngineInterface.cpp
--- a/Engine/Source/Runtime/Engine/Private/OnlineEngineInterface.cpp
+++ b/Engine/Source/Runtime/Engine/Private/OnlineEngineInterface.cpp
@@ -2,6 +2,7 @@
 
 #include "Net/OnlineEngineInterface.h"
 #include "UObject/Package.h"
+#include "Misc/ConfigCacheIni.h"
 
 UOnlineEngineInterface* UOnlineEngineInterface::Singleton = nullptr;
 
@@ -14,8 +15,18 @@ UOnlineEngineInterface* UOnlineEngineInterface::Get()
 {
 	if (!Singleton)
 	{
+		bool bUseOnlineServicesV2 = false;
+		GConfig->GetBool(TEXT("/Script/Engine.OnlineEngineInterface"), TEXT("bUseOnlineServicesV2"), bUseOnlineServicesV2, GEngineIni);
 		// Proper interface class hard coded here to emphasize the fact that this is not expected to change much, any need to do so should go through the OGS team first
-		UClass* OnlineEngineInterfaceClass = StaticLoadClass(UOnlineEngineInterface::StaticClass(), NULL, TEXT("/Script/OnlineSubsystemUtils.OnlineEngineInterfaceImpl"), NULL, LOAD_Quiet, NULL);
+		UClass* OnlineEngineInterfaceClass = nullptr;
+		if (bUseOnlineServicesV2)
+		{
+			OnlineEngineInterfaceClass = StaticLoadClass(UOnlineEngineInterface::StaticClass(), NULL, TEXT("/Script/OnlineSubsystemUtils.OnlineServicesEngineInterfaceImpl"), NULL, LOAD_Quiet, NULL);
+		}
+		else
+		{
+			OnlineEngineInterfaceClass = StaticLoadClass(UOnlineEngineInterface::StaticClass(), NULL, TEXT("/Script/OnlineSubsystemUtils.OnlineEngineInterfaceImpl"), NULL, LOAD_Quiet, NULL);
+		}
 		if (!OnlineEngineInterfaceClass)
 		{
 			// Default to the no op class if necessary
diff --git a/Engine/Source/Runtime/Engine/Private/OnlineReplStructs.cpp b/Engine/Source/Runtime/Engine/Private/OnlineReplStructs.cpp
--- a/Engine/Source/Runtime/Engine/Private/OnlineReplStructs.cpp
+++ b/Engine/Source/Runtime/Engine/Private/OnlineReplStructs.cpp
@@ -242,8 +242,8 @@ void FUniqueNetIdRepl::MakeReplicationDataV2()
 void FUniqueNetIdRepl::UniqueIdFromString(FName Type, const FString& Contents)
 {
 	// Don't need to distinguish OSS interfaces here with world because we just want the create function below
-	FUniqueNetIdPtr UniqueNetIdPtr = UOnlineEngineInterface::Get()->CreateUniquePlayerId(Contents, Type);
-	SetUniqueNetId(UniqueNetIdPtr);
+	FUniqueNetIdWrapper UniqueNetIdWrapper = UOnlineEngineInterface::Get()->CreateUniquePlayerIdWrapper(Contents, Type);
+	SetUniqueNetId(UniqueNetIdWrapper.GetUniqueNetId());
 }
 
 bool FUniqueNetIdRepl::NetSerialize(FArchive& Ar, UPackageMap* Map, bool& bOutSuccess)
@@ -602,7 +602,7 @@ void TestUniqueIdRepl(UWorld* InWorld)
 
 	bool bSetupSuccess = true;
 
-	FUniqueNetIdPtr UserId = UOnlineEngineInterface::Get()->GetUniquePlayerId(InWorld, 0);
+	FUniqueNetIdPtr UserId = UOnlineEngineInterface::Get()->GetUniquePlayerIdWrapper(InWorld, 0).GetUniqueNetId();
 
 	FUniqueNetIdRepl EmptyIdIn;
 	if (EmptyIdIn.IsValid())
@@ -618,9 +618,9 @@ void TestUniqueIdRepl(UWorld* InWorld)
 		bSetupSuccess = false;
 	}
 
-	FUniqueNetIdRepl OddStringIdIn(UOnlineEngineInterface::Get()->CreateUniquePlayerId(TEXT("abcde")));
-	FUniqueNetIdRepl NonHexStringIdIn(UOnlineEngineInterface::Get()->CreateUniquePlayerId(TEXT("thisisnothex")));
-	FUniqueNetIdRepl UpperCaseStringIdIn(UOnlineEngineInterface::Get()->CreateUniquePlayerId(TEXT("abcDEF")));
+	FUniqueNetIdRepl OddStringIdIn(UOnlineEngineInterface::Get()->CreateUniquePlayerIdWrapper(TEXT("abcde")));
+	FUniqueNetIdRepl NonHexStringIdIn(UOnlineEngineInterface::Get()->CreateUniquePlayerIdWrapper(TEXT("thisisnothex")));
+	FUniqueNetIdRepl UpperCaseStringIdIn(UOnlineEngineInterface::Get()->CreateUniquePlayerIdWrapper(TEXT("abcDEF")));
 
 #if 1
 #define WAYTOOLONG TEXT(\
@@ -634,7 +634,7 @@ void TestUniqueIdRepl(UWorld* InWorld)
 #define WAYTOOLONG TEXT("deadbeef")
 #endif
 
-	FUniqueNetIdRepl WayTooLongForHexEncodingIdIn(UOnlineEngineInterface::Get()->CreateUniquePlayerId(WAYTOOLONG));
+	FUniqueNetIdRepl WayTooLongForHexEncodingIdIn(UOnlineEngineInterface::Get()->CreateUniquePlayerIdWrapper(WAYTOOLONG));
 
 	CHECK_REPL_VALIDITY(OddStringIdIn, bSetupSuccess);
 	CHECK_REPL_VALIDITY(NonHexStringIdIn, bSetupSuccess);
@@ -642,11 +642,11 @@ void TestUniqueIdRepl(UWorld* InWorld)
 	CHECK_REPL_VALIDITY(WayTooLongForHexEncodingIdIn, bSetupSuccess);
 
 	static FName NAME_CustomOSS(TEXT("MyCustomOSS"));
-	FUniqueNetIdRepl CustomOSSIdIn(UOnlineEngineInterface::Get()->CreateUniquePlayerId(TEXT("a8d245fc-4b97-4150-a3cd-c2c91d8fc4b3"), NAME_CustomOSS));
-	FUniqueNetIdRepl CustomOSSEncodedIdIn(UOnlineEngineInterface::Get()->CreateUniquePlayerId(TEXT("0123456789abcdef"), NAME_CustomOSS));
-	FUniqueNetIdRepl CustomOSSPlusPrefixIn(UOnlineEngineInterface::Get()->CreateUniquePlayerId(TEXT("+123456"), NAME_CustomOSS));
-	FUniqueNetIdRepl CustomOSSOddIntegerStringIn(UOnlineEngineInterface::Get()->CreateUniquePlayerId(TEXT("123456789"), NAME_CustomOSS));
-	FUniqueNetIdRepl CustomOSSEvenIntegerStringIn(UOnlineEngineInterface::Get()->CreateUniquePlayerId(TEXT("1234567890"), NAME_CustomOSS));
+	FUniqueNetIdRepl CustomOSSIdIn(UOnlineEngineInterface::Get()->CreateUniquePlayerIdWrapper(TEXT("a8d245fc-4b97-4150-a3cd-c2c91d8fc4b3"), NAME_CustomOSS));
+	FUniqueNetIdRepl CustomOSSEncodedIdIn(UOnlineEngineInterface::Get()->CreateUniquePlayerIdWrapper(TEXT("0123456789abcdef"), NAME_CustomOSS));
+	FUniqueNetIdRepl CustomOSSPlusPrefixIn(UOnlineEngineInterface::Get()->CreateUniquePlayerIdWrapper(TEXT("+123456"), NAME_CustomOSS));
+	FUniqueNetIdRepl CustomOSSOddIntegerStringIn(UOnlineEngineInterface::Get()->CreateUniquePlayerIdWrapper(TEXT("123456789"), NAME_CustomOSS));
+	FUniqueNetIdRepl CustomOSSEvenIntegerStringIn(UOnlineEngineInterface::Get()->CreateUniquePlayerIdWrapper(TEXT("1234567890"), NAME_CustomOSS));
 
 	CHECK_REPL_VALIDITY(CustomOSSIdIn, bSetupSuccess);
 	CHECK_REPL_VALIDITY(CustomOSSEncodedIdIn, bSetupSuccess);
@@ -851,7 +851,7 @@ void TestUniqueIdRepl(UWorld* InWorld)
 		GConfig->GetString(TEXT("OnlineSubsystem"), TEXT("NativePlatformService"), NativePlatformService, GEngineIni) && 
 		!NativePlatformService.IsEmpty())
 	{
-		FUniqueNetIdPtr PlatformUserId = UOnlineEngineInterface::Get()->GetUniquePlayerId(InWorld, 0, FName(*NativePlatformService));
+		FUniqueNetIdPtr PlatformUserId = UOnlineEngineInterface::Get()->GetUniquePlayerIdWrapper(InWorld, 0, FName(*NativePlatformService)).GetUniqueNetId();
 
 		FUniqueNetIdRepl ValidPlatformIdIn(PlatformUserId);
 		if (!ValidPlatformIdIn.IsValid() || PlatformUserId != ValidPlatformIdIn.GetUniqueNetId() || *PlatformUserId != *ValidPlatformIdIn)
diff --git a/Engine/Source/Runtime/Engine/Private/PlayerController.cpp b/Engine/Source/Runtime/Engine/Private/PlayerController.cpp
--- a/Engine/Source/Runtime/Engine/Private/PlayerController.cpp
+++ b/Engine/Source/Runtime/Engine/Private/PlayerController.cpp
@@ -3581,7 +3581,7 @@ void APlayerController::ClientMutePlayer_Implementation(FUniqueNetIdRepl PlayerI
 	if (LP != NULL && World)
 	{
 		// Have the voice subsystem mute this player
-		UOnlineEngineInterface::Get()->MuteRemoteTalker(World, LP->GetControllerId(), *PlayerId, false);
+		UOnlineEngineInterface::Get()->MuteRemoteTalker(World, LP->GetControllerId(), PlayerId, false);
 	}
 }
 
@@ -3594,7 +3594,7 @@ void APlayerController::ClientUnmutePlayer_Implementation(FUniqueNetIdRepl Playe
 	if (LP != NULL && World)
 	{
 		// Have the voice subsystem unmute this player
-		UOnlineEngineInterface::Get()->UnmuteRemoteTalker(World, LP->GetControllerId(), *PlayerId, false);
+		UOnlineEngineInterface::Get()->UnmuteRemoteTalker(World, LP->GetControllerId(), PlayerId, false);
 	}
 }
 
@@ -3611,7 +3611,7 @@ void APlayerController::ClientUnmutePlayers_Implementation(const TArray<FUniqueN
 			if (UnmuteId.IsValid())
 			{
 				// Have the voice subsystem mute this player
-				UOnlineEngineInterface::Get()->UnmuteRemoteTalker(World, LP->GetControllerId(), *UnmuteId, false);
+				UOnlineEngineInterface::Get()->UnmuteRemoteTalker(World, LP->GetControllerId(), UnmuteId, false);
 			}
 		}
 	}
diff --git a/Engine/Source/Runtime/Engine/Private/PlayerState.cpp b/Engine/Source/Runtime/Engine/Private/PlayerState.cpp
--- a/Engine/Source/Runtime/Engine/Private/PlayerState.cpp
+++ b/Engine/Source/Runtime/Engine/Private/PlayerState.cpp
@@ -399,7 +399,7 @@ void APlayerState::RegisterPlayerWithSession(bool bWasFromInvite)
 		{
 			// Register the player as part of the session
 			const APlayerState* PlayerState = GetDefault<APlayerState>();
-			UOnlineEngineInterface::Get()->RegisterPlayer(GetWorld(), PlayerState->SessionName, *GetUniqueId(), bWasFromInvite);
+			UOnlineEngineInterface::Get()->RegisterPlayer(GetWorld(), PlayerState->SessionName, GetUniqueId(), bWasFromInvite);
 		}
 	}
 }
@@ -411,7 +411,7 @@ void APlayerState::UnregisterPlayerWithSession()
 		const APlayerState* PlayerState = GetDefault<APlayerState>();
 		if (PlayerState->SessionName != NAME_None)
 		{
-			UOnlineEngineInterface::Get()->UnregisterPlayer(GetWorld(), PlayerState->SessionName, *GetUniqueId());
+			UOnlineEngineInterface::Get()->UnregisterPlayer(GetWorld(), PlayerState->SessionName, GetUniqueId());
 		}
 	}
 }
diff --git a/Engine/Source/Runtime/Engine/Public/Net/OnlineEngineInterface.h b/Engine/Source/Runtime/Engine/Public/Net/OnlineEngineInterface.h
--- a/Engine/Source/Runtime/Engine/Public/Net/OnlineEngineInterface.h
+++ b/Engine/Source/Runtime/Engine/Public/Net/OnlineEngineInterface.h
@@ -64,7 +64,10 @@ public:
 	/** Returns the name of the default online subsystem, generally the DefaultPlatformService config value. */
 	virtual FName GetDefaultOnlineSubsystemName() const { return NAME_None; }
 	/** Returns whether the unique net id is compatible */
-	virtual bool IsCompatibleUniqueNetId(const FUniqueNetId& InUniqueNetId) const { return (InUniqueNetId.GetType() == GetDefaultOnlineSubsystemName()); }
+	UE_DEPRECATED(5.0, "Use IsCompatibleUniqueNetId with FUniqueNetIdWrapper")
+	virtual bool IsCompatibleUniqueNetId(const FUniqueNetId& InUniqueNetId) const { return IsCompatibleUniqueNetId(FUniqueNetIdWrapper(InUniqueNetId)); }
+	/** Returns whether the unique net id is compatible */
+	virtual bool IsCompatibleUniqueNetId(const FUniqueNetIdWrapper& InUniqueNetId) const { return (InUniqueNetId.IsV1() && InUniqueNetId.GetType() == GetDefaultOnlineSubsystemName()); }
 
 	/**
 	 * Utils
@@ -93,12 +96,21 @@ public:
 	 */
 
 	/** @return a proper FUniqueNetId pointer generated by the default OnlineSubsystem */
-	virtual FUniqueNetIdPtr CreateUniquePlayerId(const FString& Str, FName Type = NAME_None) { return nullptr; }
+	UE_DEPRECATED(5.0, "Use CreateUniquePlayerIdWrapper")
+	virtual FUniqueNetIdPtr CreateUniquePlayerId(const FString& Str, FName Type = NAME_None) { return CreateUniquePlayerIdWrapper(Str, Type).GetV1(); }
+	/** @return a proper FUniqueNetId pointer generated by the default OnlineSubsystem */
+	virtual FUniqueNetIdWrapper CreateUniquePlayerIdWrapper(const FString& Str, FName Type = NAME_None) { return FUniqueNetIdWrapper(); }
 	/** @return the unique id for a given local user on the given OnlineSubsystem */
-	virtual FUniqueNetIdPtr GetUniquePlayerId(UWorld* World, int32 LocalUserNum, FName Type = NAME_None) { return nullptr; }
+	UE_DEPRECATED(5.0, "Use GetUniquePlayerIdWrapper")
+	virtual FUniqueNetIdPtr GetUniquePlayerId(UWorld* World, int32 LocalUserNum, FName Type = NAME_None) { return GetUniquePlayerIdWrapper(World, LocalUserNum, Type).GetV1(); }
+	/** @return the unique id for a given local user on the given OnlineSubsystem */
+	virtual FUniqueNetIdWrapper GetUniquePlayerIdWrapper(UWorld* World, int32 LocalUserNum, FName Type = NAME_None) { return FUniqueNetIdWrapper(); }
 
 	/** @return the stored nickname on the OnlineSubsystem for a player with a given unique id */
-	virtual FString GetPlayerNickname(UWorld* World, const FUniqueNetId& UniqueId) { return TEXT("INVALID"); }
+	UE_DEPRECATED(5.0, "Use GetPlayerNickname with FUniqueNetIdWrapper")
+	virtual FString GetPlayerNickname(UWorld* World, const FUniqueNetId& UniqueId) { return GetPlayerNickname(World, FUniqueNetIdWrapper(UniqueId)); }
+	/** @return the stored nickname on the OnlineSubsystem for a player with a given unique id */
+	virtual FString GetPlayerNickname(UWorld* World, const FUniqueNetIdWrapper& UniqueId) { return TEXT("INVALID"); }
 	/** @return the stored nickname on the secondary/console OnlineSubsystem for a player with a given unique id */
 	virtual bool GetPlayerPlatformNickname(UWorld* World, int32 LocalUserNum, FString& OutNickname) { return false; }
 	
@@ -122,10 +134,28 @@ public:
 	virtual void UpdateSessionJoinability(UWorld* World, FName SessionName, bool bPublicSearchable, bool bAllowInvites, bool bJoinViaPresence, bool bJoinViaPresenceFriendsOnly) {}
 
 	/** Register a given player with the online session */
-	virtual void RegisterPlayer(UWorld* World, FName SessionName, const FUniqueNetId& UniqueId, bool bWasInvited) {}
+	UE_DEPRECATED(5.0, "Use RegisterPlayer with FUniqueNetIdWrapper")
+	virtual void RegisterPlayer(UWorld* World, FName SessionName, const FUniqueNetId& UniqueId, bool bWasInvited) { RegisterPlayer(World, SessionName, FUniqueNetIdWrapper(UniqueId), bWasInvited); }
+	/** Register a given player with the online session */
+	virtual void RegisterPlayer(UWorld* World, FName SessionName, const FUniqueNetIdWrapper& UniqueId, bool bWasInvited) {}
+	/** Unregister a given player from the online session */
+	UE_DEPRECATED(5.0, "Use UnregisterPlayer with FUniqueNetIdWrapper")
+	virtual void UnregisterPlayer(UWorld* World, FName SessionName, const FUniqueNetId& UniqueId) { UnregisterPlayer(World, SessionName, FUniqueNetIdWrapper(UniqueId)); }
+	/** Unregister a given player from the online session */
+	virtual void UnregisterPlayer(UWorld* World, FName SessionName, const FUniqueNetIdWrapper& UniqueId) {}
 	/** Unregister a given player from the online session */
-	virtual void UnregisterPlayer(UWorld* World, FName SessionName, const FUniqueNetId& UniqueId) {}
-	virtual void UnregisterPlayers(UWorld* World, FName SessionName, const TArray< FUniqueNetIdRef >& Players) {}
+	UE_DEPRECATED(5.0, "Use UnregisterPlayers with FUniqueNetIdWrapper")
+	virtual void UnregisterPlayers(UWorld* World, FName SessionName, const TArray<FUniqueNetIdRef>& Players)
+	{
+		TArray<FUniqueNetIdWrapper> PlayerIdsAsWrappers;
+		for (const FUniqueNetIdRef& PlayerId : Players)
+		{
+			PlayerIdsAsWrappers.Emplace(PlayerId);
+		}
+		UnregisterPlayers(World, SessionName, PlayerIdsAsWrappers);
+	}
+	/** Unregister a given player from the online session */
+	virtual void UnregisterPlayers(UWorld* World, FName SessionName, const TArray<FUniqueNetIdWrapper>& Players) {}
 	/** @return true if there is a valid URL that can been used to connect to a given session, false otherwise */
 	virtual bool GetResolvedConnectString(UWorld* World, FName SessionName, FString& URL) { return false; }
 
@@ -146,9 +176,15 @@ public:
 	virtual void ClearVoicePackets(UWorld* World) {}
 
 	/** Mute a remote user for a given local player */
-	virtual bool MuteRemoteTalker(UWorld* World, uint8 LocalUserNum, const FUniqueNetId& PlayerId, bool bIsSystemWide) { return false; }
+	UE_DEPRECATED(5.0, "Use MuteRemoteTalker with FUniqueNetIdWrapper")
+	virtual bool MuteRemoteTalker(UWorld* World, uint8 LocalUserNum, const FUniqueNetId& PlayerId, bool bIsSystemWide) { return MuteRemoteTalker(World, LocalUserNum, FUniqueNetIdWrapper(PlayerId), bIsSystemWide); }
+	/** Mute a remote user for a given local player */
+	virtual bool MuteRemoteTalker(UWorld* World, uint8 LocalUserNum, const FUniqueNetIdWrapper& PlayerId, bool bIsSystemWide) { return false; }
+	/** UnMute a remote user for a given local player */
+	UE_DEPRECATED(5.0, "Use UnmuteRemoteTalker with FUniqueNetIdWrapper")
+	virtual bool UnmuteRemoteTalker(UWorld* World, uint8 LocalUserNum, const FUniqueNetId& PlayerId, bool bIsSystemWide) { return UnmuteRemoteTalker(World, LocalUserNum, FUniqueNetIdWrapper(PlayerId), bIsSystemWide); }
 	/** UnMute a remote user for a given local player */
-	virtual bool UnmuteRemoteTalker(UWorld* World, uint8 LocalUserNum, const FUniqueNetId& PlayerId, bool bIsSystemWide) { return false; }
+	virtual bool UnmuteRemoteTalker(UWorld* World, uint8 LocalUserNum, const FUniqueNetIdWrapper& PlayerId, bool bIsSystemWide) { return false; }
 	/** @return number of available/registered local talkers */
 	virtual int32 GetNumLocalTalkers(UWorld* World) { return 0; }
 
