Move public online header to Online directory

diff --git a/Engine/Plugins/Online/OnlineSubsystemUtils/Source/OnlineSubsystemUtils/Private/OnlineServicesEngineInterfaceImpl.cpp b/Engine/Plugins/Online/OnlineSubsystemUtils/Source/OnlineSubsystemUtils/Private/OnlineServicesEngineInterfaceImpl.cpp
--- a/Engine/Plugins/Online/OnlineSubsystemUtils/Source/OnlineSubsystemUtils/Private/OnlineServicesEngineInterfaceImpl.cpp
+++ b/Engine/Plugins/Online/OnlineSubsystemUtils/Source/OnlineSubsystemUtils/Private/OnlineServicesEngineInterfaceImpl.cpp
@@ -1,7 +1,7 @@
 // Copyright Epic Games, Inc. All Rights Reserved.
 
 #include "OnlineServicesEngineInterfaceImpl.h"
-#include "OnlineServicesEngineUtils.h"
+#include "Online/OnlineServicesEngineUtils.h"
 #include "Misc/ConfigCacheIni.h"
 #include "Engine/NetConnection.h"
 
diff --git a/Engine/Plugins/Online/OnlineSubsystemUtils/Source/OnlineSubsystemUtils/Private/OnlineServicesEngineUtilsImpl.h b/Engine/Plugins/Online/OnlineSubsystemUtils/Source/OnlineSubsystemUtils/Private/OnlineServicesEngineUtilsImpl.h
--- a/Engine/Plugins/Online/OnlineSubsystemUtils/Source/OnlineSubsystemUtils/Private/OnlineServicesEngineUtilsImpl.h
+++ b/Engine/Plugins/Online/OnlineSubsystemUtils/Source/OnlineSubsystemUtils/Private/OnlineServicesEngineUtilsImpl.h
@@ -2,7 +2,7 @@
 
 #pragma once
 
-#include "OnlineServicesEngineUtils.h"
+#include "Online/OnlineServicesEngineUtils.h"
 #include "OnlineDelegates.h"
 #include "OnlinePIESettings.h"
 #include "Modules/ModuleManager.h"
diff --git a/Engine/Plugins/Online/OnlineSubsystemUtils/Source/OnlineSubsystemUtils/Public/Online/OnlineServicesEngineUtils.h b/Engine/Plugins/Online/OnlineSubsystemUtils/Source/OnlineSubsystemUtils/Public/Online/OnlineServicesEngineUtils.h
new file mode 100644
--- /dev/null
+++ b/Engine/Plugins/Online/OnlineSubsystemUtils/Source/OnlineSubsystemUtils/Public/Online/OnlineServicesEngineUtils.h
@@ -0,0 +1,75 @@
+// Copyright Epic Games, Inc. All Rights Reserved.
+
+#pragma once
+
+#include "CoreMinimal.h"
+#include "Online/CoreOnline.h"
+#include "Online/Auth.h"
+
+class UWorld;
+struct FWorldContext;
+struct FPIELoginSettingsInternal;
+
+DECLARE_MULTICAST_DELEGATE_OneParam(FOnExternalUIChange, bool);
+typedef FOnExternalUIChange::FDelegate FOnExternalUIChangeDelegate;
+
+namespace UE::Online {
+
+/**
+ * Interface class for various online utility functions for Engine
+ */
+class IOnlineServicesEngineUtils
+{
+protected:
+	/** Hidden on purpose */
+	IOnlineServicesEngineUtils() {}
+public:
+
+	virtual ~IOnlineServicesEngineUtils() {}
+
+	/**
+	 * Gets an FName that uniquely identifies an instance of a world to differentiate online services for PIE
+	 *
+	 * @param WorldContext the worldcontext
+	 * @return an FName that can be used with online services
+	 */
+	virtual FName GetOnlineIdentifier(const FWorldContext& WorldContext) const = 0;
+
+	/**
+	 * Gets an FName that uniquely identifies an instance of a world to differentiate online services for PIE
+	 *
+	 * @param World the world to use for context
+	 * @return an FName that can be used with Online Services
+	 */
+	virtual FName GetOnlineIdentifier(const UWorld* World) const = 0;
+
+	/**
+	 * Bind a notification delegate when any subsystem external UI is opened/closed
+	 * *NOTE* there is only meant to be one delegate needed for this, game code should bind manually
+	 *
+	 * @param OnExternalUIChangeDelegate delegate fired when the external UI is opened/closed
+	 */
+	virtual void SetEngineExternalUIBinding(const FOnExternalUIChangeDelegate& OnExternalUIChangeDelegate) = 0;
+
+#if WITH_EDITOR
+	/**
+	 * Play in Editor settings
+	 */
+
+	/** @return true if the default platform supports logging in for Play In Editor (PIE) */
+	virtual bool SupportsOnlinePIE() const = 0;
+	/** Enable/Disable online PIE at runtime */
+	virtual void SetShouldTryOnlinePIE(bool bShouldTry) = 0;
+	/** @return true if the user has enabled logging in for Play In Editor (PIE) */
+	virtual bool IsOnlinePIEEnabled() const = 0;
+	/** @return the number of logins the user has setup for Play In Editor (PIE) */
+	virtual int32 GetNumPIELogins() const = 0;
+	/** @return the array of valid credentials the user has setup for Play In Editor (PIE) */
+	virtual void GetPIELogins(TArray<FAuthLogin::Params>& Logins) const = 0;
+#endif // WITH_EDITOR
+};
+
+/** @return the single instance of the online services utils interface */
+ONLINESUBSYSTEMUTILS_API IOnlineServicesEngineUtils* GetServicesEngineUtils();
+
+/* UE::Online */ }
diff --git a/Engine/Plugins/Online/OnlineSubsystemUtils/Source/OnlineSubsystemUtils/Public/OnlineServicesEngineUtils.h b/Engine/Plugins/Online/OnlineSubsystemUtils/Source/OnlineSubsystemUtils/Public/OnlineServicesEngineUtils.h
deleted file mode 100644
