Fix shutdown crash in ~FOnlineServicesEngineUtils.
Set OnlineSubsystemUtils to support automatic shutdown so ShutdownModule get called to cleanup before Online Services is unloaded.

diff --git a/Engine/Plugins/Online/OnlineSubsystemUtils/Source/OnlineSubsystemUtils/Public/OnlineSubsystemUtilsModule.h b/Engine/Plugins/Online/OnlineSubsystemUtils/Source/OnlineSubsystemUtils/Public/OnlineSubsystemUtilsModule.h
--- a/Engine/Plugins/Online/OnlineSubsystemUtils/Source/OnlineSubsystemUtils/Public/OnlineSubsystemUtilsModule.h
+++ b/Engine/Plugins/Online/OnlineSubsystemUtils/Source/OnlineSubsystemUtils/Public/OnlineSubsystemUtilsModule.h
@@ -36,12 +36,6 @@ public:
 		return false;
 	}
 
-	virtual bool SupportsAutomaticShutdown() override
-	{
-		return false;
-	}
-
-
 private:
 
 	/** Online subsystem utilities singleton */
