Add GetServices accessor that takes a world
Use Default OnlineService type in utils

diff --git a/Engine/Plugins/Online/OnlineSubsystemUtils/Source/OnlineSubsystemUtils/Private/OnlineServicesEngineInterfaceImpl.cpp b/Engine/Plugins/Online/OnlineSubsystemUtils/Source/OnlineSubsystemUtils/Private/OnlineServicesEngineInterfaceImpl.cpp
--- a/Engine/Plugins/Online/OnlineSubsystemUtils/Source/OnlineSubsystemUtils/Private/OnlineServicesEngineInterfaceImpl.cpp
+++ b/Engine/Plugins/Online/OnlineSubsystemUtils/Source/OnlineSubsystemUtils/Private/OnlineServicesEngineInterfaceImpl.cpp
@@ -20,7 +20,7 @@ UOnlineServicesEngineInterfaceImpl::UOnlineServicesEngineInterfaceImpl(const FOb
 bool UOnlineServicesEngineInterfaceImpl::IsLoaded(FName OnlineIdentifier)
 {
 	// TODO: something like OnlineServicesTypeFromName(OnlineIdentifier); ? This seems to only be called with NAME_None
-	UE::Online::EOnlineServices OnlineServicesType = UE::Online::EOnlineServices::Epic; // TODO:  Need Default support UE::Online::EOnlineServices::Default
+	UE::Online::EOnlineServices OnlineServicesType = UE::Online::EOnlineServices::Default;
 	return UE::Online::IsLoaded(OnlineServicesType, OnlineIdentifier);
 }
 
@@ -103,7 +103,7 @@ FUniqueNetIdWrapper UOnlineServicesEngineInterfaceImpl::GetUniquePlayerIdWrapper
 {
 	// TODO: something like OnlineServicesTypeFromName(Type); ? This seems to only be called with NAME_None in LocalPlayer and with the platform service in OnlineReplStructs in a test command
 	FName OnlineIdentifier = GetOnlineIdentifier(World);
-	UE::Online::EOnlineServices OnlineServicesType = UE::Online::EOnlineServices::Epic; // TODO:  Need Default support UE::Online::EOnlineServices::Default
+	UE::Online::EOnlineServices OnlineServicesType = UE::Online::EOnlineServices::Default;
 	if (UE::Online::IOnlineServicesPtr OnlineServices = UE::Online::GetServices(OnlineServicesType, OnlineIdentifier))
 	{
 		if (UE::Online::IAuthPtr AuthPtr = OnlineServices->GetAuthInterface())
@@ -123,7 +123,7 @@ FString UOnlineServicesEngineInterfaceImpl::GetPlayerNickname(UWorld* World, con
 {
 	check(UniqueId.IsValid() && UniqueId.IsV2());
 	FName OnlineIdentifier = GetOnlineIdentifier(World);
-	UE::Online::EOnlineServices OnlineServicesType = UE::Online::EOnlineServices::Epic; // TODO: Get from UniqueId
+	UE::Online::EOnlineServices OnlineServicesType = UE::Online::EOnlineServices::Default; // TODO: Get from UniqueId
 	if (UE::Online::IOnlineServicesPtr OnlineServices = UE::Online::GetServices(OnlineServicesType, OnlineIdentifier))
 	{
 		if (UE::Online::IAuthPtr AuthPtr = OnlineServices->GetAuthInterface())
@@ -163,7 +163,7 @@ bool UOnlineServicesEngineInterfaceImpl::GetPlayerPlatformNickname(UWorld* World
 bool UOnlineServicesEngineInterfaceImpl::AutoLogin(UWorld* World, int32 LocalUserNum, const FOnlineAutoLoginComplete& InCompletionDelegate)
 {
 	FName OnlineIdentifier = GetOnlineIdentifier(World);
-	UE::Online::EOnlineServices OnlineServicesType = UE::Online::EOnlineServices::Epic; // TODO:  Need Default support UE::Online::EOnlineServices::Default
+	UE::Online::EOnlineServices OnlineServicesType = UE::Online::EOnlineServices::Default;
 	if (UE::Online::IOnlineServicesPtr OnlineServices = UE::Online::GetServices(OnlineServicesType, OnlineIdentifier))
 	{
 		if (UE::Online::IAuthPtr AuthPtr = OnlineServices->GetAuthInterface())
@@ -187,7 +187,7 @@ bool UOnlineServicesEngineInterfaceImpl::AutoLogin(UWorld* World, int32 LocalUse
 bool UOnlineServicesEngineInterfaceImpl::IsLoggedIn(UWorld* World, int32 LocalUserNum)
 {
 	FName OnlineIdentifier = GetOnlineIdentifier(World);
-	UE::Online::EOnlineServices OnlineServicesType = UE::Online::EOnlineServices::Epic; // TODO:  Need Default support UE::Online::EOnlineServices::Default
+	UE::Online::EOnlineServices OnlineServicesType = UE::Online::EOnlineServices::Default;
 	if (UE::Online::IOnlineServicesPtr OnlineServices = UE::Online::GetServices(OnlineServicesType, OnlineIdentifier))
 	{
 		if (UE::Online::IAuthPtr AuthPtr = OnlineServices->GetAuthInterface())
@@ -309,7 +309,7 @@ int32 UOnlineServicesEngineInterfaceImpl::GetNumLocalTalkers(UWorld* World)
 void UOnlineServicesEngineInterfaceImpl::ShowLeaderboardUI(UWorld* World, const FString& CategoryName)
 {
 	FName OnlineIdentifier = GetOnlineIdentifier(World);
-	UE::Online::EOnlineServices OnlineServicesType = UE::Online::EOnlineServices::Epic; // TODO:  Need Default support UE::Online::EOnlineServices::Default
+	UE::Online::EOnlineServices OnlineServicesType = UE::Online::EOnlineServices::Default;
 	if (UE::Online::IOnlineServicesPtr OnlineServices = UE::Online::GetServices(OnlineServicesType, OnlineIdentifier))
 	{
 		UE::Online::IExternalUIPtr ExternalUI = OnlineServices->GetExternalUIInterface();
@@ -416,7 +416,7 @@ int32 UOnlineServicesEngineInterfaceImpl::GetNumPIELogins()
 void UOnlineServicesEngineInterfaceImpl::SetForceDedicated(FName OnlineIdentifier, bool bForce)
 {
 	// TODO:  Support other services? This is only called from PlayLevel with the result of GetOnlineIdentifier which only returns the default service id
-	UE::Online::EOnlineServices OnlineServicesType = UE::Online::EOnlineServices::Epic; // TODO:  Need Default support UE::Online::EOnlineServices::Default
+	UE::Online::EOnlineServices OnlineServicesType = UE::Online::EOnlineServices::Default;
 	if (UE::Online::IOnlineServicesPtr OnlineServices = UE::Online::GetServices(OnlineServicesType, OnlineIdentifier))
 	{
 		// TODO:  Add SetForceDedicated method to OnlineServices
@@ -433,7 +433,7 @@ void UOnlineServicesEngineInterfaceImpl::LoginPIEInstance(FName OnlineIdentifier
 		if (PIELogins.IsValidIndex(PIELoginNum))
 		{
 			// TODO:  Support other services? This is only called from PlayLevel with the result of GetOnlineIdentifier which only returns the default service id
-			UE::Online::EOnlineServices OnlineServicesType = UE::Online::EOnlineServices::Epic; // TODO:  Need Default support UE::Online::EOnlineServices::Default
+			UE::Online::EOnlineServices OnlineServicesType = UE::Online::EOnlineServices::Default;
 			if (UE::Online::IOnlineServicesPtr OnlineServices = UE::Online::GetServices(OnlineServicesType, OnlineIdentifier))
 			{
 				if (UE::Online::IAuthPtr AuthPtr = OnlineServices->GetAuthInterface())
diff --git a/Engine/Plugins/Online/OnlineSubsystemUtils/Source/OnlineSubsystemUtils/Public/Online/OnlineServicesEngineUtils.h b/Engine/Plugins/Online/OnlineSubsystemUtils/Source/OnlineSubsystemUtils/Public/Online/OnlineServicesEngineUtils.h
--- a/Engine/Plugins/Online/OnlineSubsystemUtils/Source/OnlineSubsystemUtils/Public/Online/OnlineServicesEngineUtils.h
+++ b/Engine/Plugins/Online/OnlineSubsystemUtils/Source/OnlineSubsystemUtils/Public/Online/OnlineServicesEngineUtils.h
@@ -5,6 +5,7 @@
 #include "CoreMinimal.h"
 #include "Online/CoreOnline.h"
 #include "Online/Auth.h"
+#include "Online/OnlineServices.h"
 
 class UWorld;
 struct FWorldContext;
@@ -72,4 +73,18 @@ public:
 /** @return the single instance of the online services utils interface */
 ONLINESUBSYSTEMUTILS_API IOnlineServicesEngineUtils* GetServicesEngineUtils();
 
+inline IOnlineServicesPtr GetServices(const UWorld* World, EOnlineServices OnlineServices = EOnlineServices::Default)
+{
+	FName Identifier;
+#if UE_EDITOR // at present, multiple worlds are only possible in the editor
+	if (World != nullptr)
+	{
+		IOnlineServicesEngineUtils* Utils = GetServicesEngineUtils();
+		check(Utils);
+		Identifier = Utils->GetOnlineIdentifier(World);
+	}
+#endif
+	return GetServices(OnlineServices, Identifier);
+}
+
 /* UE::Online */ }
