Fix OSSv2 issues

diff --git a/Engine/Plugins/Online/OnlineSubsystemUtils/Source/OnlineSubsystemUtils/Private/OnlineServicesEngineInterfaceImpl.cpp b/Engine/Plugins/Online/OnlineSubsystemUtils/Source/OnlineSubsystemUtils/Private/OnlineServicesEngineInterfaceImpl.cpp
--- a/Engine/Plugins/Online/OnlineSubsystemUtils/Source/OnlineSubsystemUtils/Private/OnlineServicesEngineInterfaceImpl.cpp
+++ b/Engine/Plugins/Online/OnlineSubsystemUtils/Source/OnlineSubsystemUtils/Private/OnlineServicesEngineInterfaceImpl.cpp
@@ -128,7 +128,7 @@ FString UOnlineServicesEngineInterfaceImpl::GetPlayerNickname(UWorld* World, con
 	{
 		if (UE::Online::IAuthPtr AuthPtr = OnlineServices->GetAuthInterface())
 		{
-			UE::Online::FAuthGetAccountByAccountId::Params GetAccountParams = { /* UniqueId.GetV2() */ };
+			UE::Online::FAuthGetAccountByAccountId::Params GetAccountParams = { UniqueId.GetV2() };
 			UE::Online::TOnlineResult<UE::Online::FAuthGetAccountByAccountId> GetAccountResult = AuthPtr->GetAccountByAccountId(MoveTemp(GetAccountParams));
 			if (GetAccountResult.IsOk())
 			{
diff --git a/Engine/Source/Runtime/Engine/Private/GameModeBase.cpp b/Engine/Source/Runtime/Engine/Private/GameModeBase.cpp
--- a/Engine/Source/Runtime/Engine/Private/GameModeBase.cpp
+++ b/Engine/Source/Runtime/Engine/Private/GameModeBase.cpp
@@ -974,7 +974,7 @@ void AGameModeBase::PostLogin(APlayerController* NewPlayer)
 	{
 		// If NewPlayer is not only a spectator and has a valid ID, add him as a user to the replay.
 		const FUniqueNetIdRepl& NewPlayerStateUniqueId = NewPlayer->PlayerState->GetUniqueId();
-		if (NewPlayerStateUniqueId.IsValid())
+		if (NewPlayerStateUniqueId.IsValid() && NewPlayerStateUniqueId.IsV1())
 		{
 			GetGameInstance()->AddUserToReplay(NewPlayerStateUniqueId.ToString());
 		}
