Reduced NotifyAcceptingConnection logspam, aggregating logs after a threshold, while avoding spam during DoS/DDoS.

diff --git a/Engine/Plugins/Online/OnlineSubsystemUtils/Source/OnlineSubsystemUtils/Classes/IpNetDriver.h b/Engine/Plugins/Online/OnlineSubsystemUtils/Source/OnlineSubsystemUtils/Classes/IpNetDriver.h
--- a/Engine/Plugins/Online/OnlineSubsystemUtils/Source/OnlineSubsystemUtils/Classes/IpNetDriver.h
+++ b/Engine/Plugins/Online/OnlineSubsystemUtils/Source/OnlineSubsystemUtils/Classes/IpNetDriver.h
@@ -222,6 +222,33 @@ private:
 		return Error == SE_NO_ERROR || Error == SE_EWOULDBLOCK || Error == SE_ECONNABORTED || Error == SE_ENETDOWN;
 	};
 
+	/**
+	 * When a new pre-connection IP is encountered, do rate-limited logging/tracking for it.
+	 *
+	 * @param InAddr	The IP address to log/track
+	 */
+	void TrackAndLogNewIP(const FInternetAddr& InAddr);
+
+	/**
+	 * Removes an IP from new IP tracking, after successful connection.
+	 *
+	 * @param InAddr	The IP address to remove from tracking
+	 */
+	void RemoveFromNewIPTracking(const FInternetAddr& InAddr);
+
+	/**
+	 * Resets new pre-connection IP tracking and timers.
+	 */
+	void ResetNewIPTracking();
+
+	/**
+	 * When new pre-connection IP tracking is enabled, ticks the tracking
+	 *
+	 * @param DeltaTime		The time since the last global tick
+	 */
+	void TickNewIPTracking(float DeltaTime);
+
+
 public:
 	// Callback for platform handling when networking is taking a long time in a single frame (by default over 1 second).
 	// It may get called multiple times in a single frame if additional processing after a previous alert exceeds the threshold again
@@ -318,4 +345,30 @@ private:
 
 	/** Underlying socket communication */
 	TSharedPtr<FSocket> SocketPrivate;
+
+
+	/** New IP aggregated logging entry */
+	struct FAggregatedIP
+	{
+		/** Hash of the IP */
+		uint32 IPHash;
+
+		/** IP:Port converted to string */
+		FString IPStr;
+	};
+
+	/** Hashes for new pre-connection IP's that are being tracked */
+	TArray<uint32> NewIPHashes;
+
+	/** Number of times each 'NewIPHashes' IP was encountered, in the current tracking period */
+	TArray<uint32> NewIPHitCount;
+
+	/** List of IP's queued for aggregated logging */
+	TArray<FAggregatedIP> AggregatedIPsToLog;
+
+	/** Whether or not IP hash tracking or aggregated IP limits have been reached. Disables all further IP logging, for the current tracking period */
+	bool bExceededIPAggregationLimit = false;
+
+	/** Countdown timer for the current IP aggregation tracking period */
+	double NextAggregateIPLogCountdown = 0.0;
 };
diff --git a/Engine/Plugins/Online/OnlineSubsystemUtils/Source/OnlineSubsystemUtils/Private/IpNetDriver.cpp b/Engine/Plugins/Online/OnlineSubsystemUtils/Source/OnlineSubsystemUtils/Private/IpNetDriver.cpp
--- a/Engine/Plugins/Online/OnlineSubsystemUtils/Source/OnlineSubsystemUtils/Private/IpNetDriver.cpp
+++ b/Engine/Plugins/Online/OnlineSubsystemUtils/Source/OnlineSubsystemUtils/Private/IpNetDriver.cpp
@@ -26,6 +26,8 @@ Notes:
 #include "IPAddress.h"
 #include "Sockets.h"
 #include "Serialization/ArchiveCountMem.h"
+#include "Algo/IndexOf.h"
+#include <limits>
 
 /** For backwards compatibility with the engine stateless connect code */
 #ifndef STATELESSCONNECT_HAS_RANDOM_SEQUENCE
@@ -128,6 +130,44 @@ TAutoConsoleVariable<FString> CVarNetDebugAddResolverAddress(
 
 namespace IPNetDriverInternal
 {
+	/** The maximum number of times to individually log a specific IP pre-connection, before aggregating further logs */
+	static int32 MaxIPHitLogs = 4;
+
+	static FAutoConsoleVariableRef CVarMaxIPHitLogs(
+		TEXT("net.MaxIPHitLogs"),
+		MaxIPHitLogs,
+		TEXT("The maximum number of times to individually log a specific IP pre-connection, before aggregating further logs."),
+		FConsoleVariableDelegate::CreateLambda([](IConsoleVariable* Var)
+		{
+			// Safe to update static here, but not CVar system cached value
+			MaxIPHitLogs = FMath::Clamp(MaxIPHitLogs, 0, std::numeric_limits<int32>::max());
+		}));
+
+	/** The maximum number of IP's to include in aggregated pre-connection logging, before such logging is disabled altogether */
+	static int32 MaxAggregateIPLogs = 16;
+
+	static FAutoConsoleVariableRef CVarMaxAggregateIPLogs(
+		TEXT("net.MaxAggregateIPLogs"),
+		MaxAggregateIPLogs,
+		TEXT("The maximum number of IP's to include in aggregated pre-connection logging, before such logging is disabled altogether ")
+		TEXT("(Min: 1, Max: 128)."),
+		FConsoleVariableDelegate::CreateLambda([](IConsoleVariable* Var)
+		{
+			// Safe to update static here, but not CVar system cached value
+			MaxAggregateIPLogs = FMath::Clamp(MaxAggregateIPLogs, 1, 128);
+		}));
+
+
+	/** The maximum number of IP hashes to track, for pre-connection logging */
+	constexpr int32 MaxIPHashes = 256;
+
+	/** The base time interval within which pre-connection IP tracking is performed */
+	constexpr double AggregateIPLogInterval = 15.0;
+
+	/** The random time interval variance added to AggregateIPLogInterval */
+	constexpr double AggregateIPLogIntervalVariance = 5.0;
+
+
 	bool ShouldSleepOnWaitError(ESocketErrors SocketError)
 	{
 		return SocketError == ESocketErrors::SE_NO_ERROR || SocketError == ESocketErrors::SE_EWOULDBLOCK || SocketError == ESocketErrors::SE_TRY_AGAIN;
@@ -325,7 +365,7 @@ private:
 	 */
 	void AdvanceCurrentPacket()
 	{
-		// @todo: Remove the slow frame checks, eventually - potential DDoS and Switch platform constraint
+		// @todo: Remove the slow frame checks, eventually - potential DDoS and platform constraint
 		if (bSlowFrameChecks)
 		{
 			const double CurrentTime = FPlatformTime::Seconds();
@@ -740,6 +780,11 @@ UIpNetDriver::UIpNetDriver(const FObjectInitializer& ObjectInitializer)
 	, ClientDesiredSocketSendBufferBytes(0x8000)
 	, RecvMultiState(nullptr)
 {
+	using namespace IPNetDriverInternal;
+
+	NewIPHashes.Empty(MaxIPHashes);
+	NewIPHitCount.Empty(MaxIPHashes);
+	AggregatedIPsToLog.Empty(MaxAggregateIPLogs);
 }
 
 bool UIpNetDriver::IsAvailable() const
@@ -1365,8 +1410,10 @@ void UIpNetDriver::TickDispatch(float DeltaTime)
 
 				if (bAcceptingConnection)
 				{
-					UE_CLOG(!DDoS.CheckLogRestrictions(), LogNet, Log, TEXT("NotifyAcceptingConnection accepted from: %s"),
-								*FromAddr->ToString(true));
+					if (!DDoS.CheckLogRestrictions() && !bExceededIPAggregationLimit)
+					{
+						TrackAndLogNewIP(FromAddr.Get());
+					}
 
 					FPacketBufferView WorkingBuffer = It.GetWorkingBuffer();
 
@@ -1398,6 +1445,11 @@ void UIpNetDriver::TickDispatch(float DeltaTime)
 		}
 	}
 
+	if (NewIPHashes.Num() > 0)
+	{
+		TickNewIPTracking(DeltaTime);
+	}
+
 	DDoS.PostFrameReceive();
 }
 
@@ -1465,6 +1517,8 @@ UNetConnection* UIpNetDriver::ProcessConnectionlessPacket(FReceivedPacketView& P
 						// @todo: There needs to be a proper/standardized copy API for this. Also in IpConnection.cpp
 						bool bIsValid = false;
 
+						RemoveFromNewIPTracking(RemoteAddrRef.Get());
+
 						const FString OldAddress = RemoteAddrRef->ToString(true);
 
 						RemoteAddrRef->SetIp(*Address->ToString(false), bIsValid);
@@ -1568,7 +1622,9 @@ UNetConnection* UIpNetDriver::ProcessConnectionlessPacket(FReceivedPacketView& P
 			}
 
 			Notify->NotifyAcceptedConnection(ReturnVal);
+
 			AddClientConnection(ReturnVal);
+			RemoveFromNewIPTracking(*Address.Get());
 		}
 
 		if (StatelessConnect.IsValid())
@@ -1793,6 +1849,149 @@ UIpConnection* UIpNetDriver::GetServerConnection()
 	return Cast<UIpConnection>(ServerConnection);
 }
 
+void UIpNetDriver::TrackAndLogNewIP(const FInternetAddr& InAddr)
+{
+	using namespace IPNetDriverInternal;
+
+	const uint32 IPHash = GetTypeHash(InAddr);
+	int32 HashIdx = NewIPHashes.Find(IPHash);
+
+	if (HashIdx == INDEX_NONE)
+	{
+		bExceededIPAggregationLimit = NewIPHashes.Num() >= MaxIPHashes;
+
+		if (LIKELY(!bExceededIPAggregationLimit))
+		{
+			HashIdx = NewIPHashes.Num();
+
+			NewIPHashes.Add(IPHash);
+			NewIPHitCount.Add(0);
+		}
+	}
+
+	if (LIKELY(!bExceededIPAggregationLimit))
+	{
+		const uint32 IPHitCount = ++NewIPHitCount[HashIdx];
+		const bool bLogIPHit = IPHitCount <= static_cast<uint32>(MaxIPHitLogs);
+
+		if (bLogIPHit)
+		{
+			UE_LOG(LogNet, Log, TEXT("NotifyAcceptingConnection accepted from: %s"), ToCStr(InAddr.ToString(true)));
+		}
+		else if (IPHitCount == static_cast<uint32>(MaxIPHitLogs) + 1)
+		{
+			bExceededIPAggregationLimit = AggregatedIPsToLog.Num() >= MaxAggregateIPLogs;
+
+			if (!bExceededIPAggregationLimit)
+			{
+				AggregatedIPsToLog.Add(UIpNetDriver::FAggregatedIP{IPHash, InAddr.ToString(true)});
+			}
+		}
+	}
+}
+
+void UIpNetDriver::RemoveFromNewIPTracking(const FInternetAddr& InAddr)
+{
+	using namespace IPNetDriverInternal;
+
+	const uint32 IPHash = GetTypeHash(InAddr);
+	const int32 HashIdx = NewIPHashes.Find(IPHash);
+
+	if (HashIdx != INDEX_NONE)
+	{
+		const bool bInAggregateLogList = NewIPHitCount[HashIdx] > static_cast<uint32>(MaxIPHitLogs);
+
+		if (bInAggregateLogList)
+		{
+			int32 AggIdx = Algo::IndexOfBy(AggregatedIPsToLog, IPHash, &UIpNetDriver::FAggregatedIP::IPHash);
+
+			if (AggIdx != INDEX_NONE)
+			{
+				AggregatedIPsToLog.RemoveAtSwap(AggIdx, 1, false);
+			}
+		}
+
+		// Don't remove from other tracking arrays unless it's the last entry
+		if (NewIPHashes.Num() == 1)
+		{
+			ResetNewIPTracking();
+		}
+		else if (HashIdx == NewIPHashes.Num()-1)
+		{
+			NewIPHashes.RemoveAt(HashIdx, 1, false);
+			NewIPHitCount.RemoveAt(HashIdx, 1, false);
+		}
+		else
+		{
+			NewIPHitCount[HashIdx] = 0;
+		}
+	}
+}
+
+void UIpNetDriver::ResetNewIPTracking()
+{
+	using namespace IPNetDriverInternal;
+
+	NewIPHashes.Empty(MaxIPHashes);
+	NewIPHitCount.Empty(MaxIPHashes);
+	AggregatedIPsToLog.Empty(MaxAggregateIPLogs);
+	bExceededIPAggregationLimit = false;
+	NextAggregateIPLogCountdown = 0.0;
+}
+
+void UIpNetDriver::TickNewIPTracking(float DeltaTime)
+{
+	using namespace IPNetDriverInternal;
+
+	bool bCheckAggregatedIPs = false;
+
+	if (NextAggregateIPLogCountdown == 0.0)
+	{
+		NextAggregateIPLogCountdown = AggregateIPLogInterval + FMath::RandRange(0.0, AggregateIPLogIntervalVariance);
+	}
+	else
+	{
+		NextAggregateIPLogCountdown -= DeltaTime;
+
+		if (NextAggregateIPLogCountdown <= 0.0)
+		{
+			NextAggregateIPLogCountdown = 0.0;
+			bCheckAggregatedIPs = true;
+		}
+	}
+
+	if (bCheckAggregatedIPs)
+	{
+		if (bExceededIPAggregationLimit)
+		{
+			UE_LOG(LogNet, Log, TEXT("NotifyAcceptingConnection accepted aggregation: Last tracking period exceeded limit."));
+		}
+		else if (AggregatedIPsToLog.Num() > 0)
+		{
+			TStringBuilder<8192> AggregateIPLogStr;
+
+			AggregateIPLogStr.Append(TEXT("NotifyAcceptingConnection accepted aggregation: "));
+
+			for (int32 AggIdx=0; AggIdx<AggregatedIPsToLog.Num(); AggIdx++)
+			{
+				const int32 HashIdx = NewIPHashes.Find(AggregatedIPsToLog[AggIdx].IPHash);
+
+				if (AggIdx > 0)
+				{
+					AggregateIPLogStr.Append(TEXT(", "));
+				}
+
+				AggregateIPLogStr.Append(ToCStr(AggregatedIPsToLog[AggIdx].IPStr));
+				AggregateIPLogStr.Appendf(TEXT(" (%i)"), NewIPHitCount[HashIdx]);
+			}
+
+			UE_LOG(LogNet, Log, TEXT("%s"), ToCStr(AggregateIPLogStr.ToString()));
+		}
+
+		ResetNewIPTracking();
+	}
+}
+
 UIpNetDriver::FReceiveThreadRunnable::FReceiveThreadRunnable(UIpNetDriver* InOwningNetDriver)
 	: ReceiveQueue(CVarNetIpNetDriverReceiveThreadQueueMaxPackets.GetValueOnAnyThread())
 	, bIsRunning(true)
diff --git a/Engine/Source/Runtime/Sockets/Public/IPAddress.h b/Engine/Source/Runtime/Sockets/Public/IPAddress.h
--- a/Engine/Source/Runtime/Sockets/Public/IPAddress.h
+++ b/Engine/Source/Runtime/Sockets/Public/IPAddress.h
@@ -420,3 +420,8 @@ struct FInternetAddrConstKeyMapFuncs : public BaseKeyFuncs<ValueType, TSharedRef
 		return Key->GetTypeHash();
 	}
 };
+
+inline uint32 GetTypeHash(const FInternetAddr& InAddr)
+{
+	return InAddr.GetTypeHash();
+}
