Fix CIS error.

diff --git a/Engine/Plugins/Online/OnlineSubsystemUtils/Source/OnlineSubsystemUtils/Private/IpNetDriver.cpp b/Engine/Plugins/Online/OnlineSubsystemUtils/Source/OnlineSubsystemUtils/Private/IpNetDriver.cpp
--- a/Engine/Plugins/Online/OnlineSubsystemUtils/Source/OnlineSubsystemUtils/Private/IpNetDriver.cpp
+++ b/Engine/Plugins/Online/OnlineSubsystemUtils/Source/OnlineSubsystemUtils/Private/IpNetDriver.cpp
@@ -1903,7 +1903,7 @@ void UIpNetDriver::RemoveFromNewIPTracking(const FInternetAddr& InAddr)
 
 		if (bInAggregateLogList)
 		{
-			int32 AggIdx = Algo::IndexOfBy(AggregatedIPsToLog, IPHash, &UIpNetDriver::FAggregatedIP::IPHash);
+			const int32 AggIdx = AggregatedIPsToLog.IndexOfByPredicate([IPHash](const FAggregatedIP& A) { return A.IPHash == IPHash; });
 
 			if (AggIdx != INDEX_NONE)
 			{
