- Rename VoiceInterfaceImpl stat scope and tighten its scope.

diff --git a/Engine/Plugins/Online/OnlineSubsystemUtils/Source/OnlineSubsystemUtils/Private/VoiceInterfaceImpl.cpp b/Engine/Plugins/Online/OnlineSubsystemUtils/Source/OnlineSubsystemUtils/Private/VoiceInterfaceImpl.cpp
--- a/Engine/Plugins/Online/OnlineSubsystemUtils/Source/OnlineSubsystemUtils/Private/VoiceInterfaceImpl.cpp
+++ b/Engine/Plugins/Online/OnlineSubsystemUtils/Source/OnlineSubsystemUtils/Private/VoiceInterfaceImpl.cpp
@@ -117,14 +117,14 @@ void FOnlineVoiceImpl::Tick(float DeltaTime)
 {
 	if (!OnlineSubsystem->IsDedicated())
 	{
-		SCOPE_CYCLE_COUNTER(STAT_Voice_Interface);
-
 		// If we aren't in a networked match, no need to update networked voice
 		if (SessionInt && SessionInt->GetNumSessions() > 0)
 		{
 			// Processing voice data only valid with a voice engine to capture/play
 			if (VoiceEngine.IsValid())
 			{
+				QUICK_SCOPE_CYCLE_COUNTER(STAT_FOnlineVoiceImpl_Tick);
+
 				VoiceEngine->Tick(DeltaTime);
 
 				// Queue local packets for sending via the network
