Changed UpdateMuteListForLocalTalker's prototype as it now retrieves the Player Controller inside its body.

diff --git a/Engine/Plugins/Online/OnlineSubsystemUtils/Source/OnlineSubsystemUtils/Private/VoiceInterfaceImpl.cpp b/Engine/Plugins/Online/OnlineSubsystemUtils/Source/OnlineSubsystemUtils/Private/VoiceInterfaceImpl.cpp
--- a/Engine/Plugins/Online/OnlineSubsystemUtils/Source/OnlineSubsystemUtils/Private/VoiceInterfaceImpl.cpp
+++ b/Engine/Plugins/Online/OnlineSubsystemUtils/Source/OnlineSubsystemUtils/Private/VoiceInterfaceImpl.cpp
@@ -550,15 +550,8 @@ void FOnlineVoiceImpl::ProcessMuteChangeNotification()
 			// For each local user with voice
 			for (int32 Index = 0; Index < MaxLocalTalkers; Index++)
 			{
-				// Find the very first ULocalPlayer for this ControllerId. 
-				// This is imperfect and means we cannot support voice chat properly for
-				// multiple UWorlds (but thats ok for the time being).
-				ULocalPlayer* LP = GEngine->FindFirstLocalPlayerFromControllerId(Index);
-				if (LP && LP->PlayerController)
-				{
-					// Use the common method of checking muting
-					UpdateMuteListForLocalTalker(Index, LP->PlayerController);
-				}
+				// Use the common method of checking muting
+				UpdateMuteListForLocalTalker(Index);
 			}
 		}
 	}
@@ -569,29 +562,37 @@ IVoiceEnginePtr FOnlineVoiceImpl::CreateVoiceEngine()
 	return MakeShareable(new FVoiceEngineImpl(OnlineSubsystem));
 }
 
-void FOnlineVoiceImpl::UpdateMuteListForLocalTalker(int32 TalkerIndex, APlayerController* PlayerController)
+void FOnlineVoiceImpl::UpdateMuteListForLocalTalker(int32 TalkerIndex)
 {
-	// For each registered remote talker
-	for (int32 RemoteIndex = 0; RemoteIndex < RemoteTalkers.Num(); RemoteIndex++)
+	// Find the very first ULocalPlayer for this ControllerId. 
+	// This is imperfect and means we cannot support voice chat properly for
+	// multiple UWorlds (but thats ok for the time being).
+	ULocalPlayer* LP = GEngine->FindFirstLocalPlayerFromControllerId(TalkerIndex);
+	APlayerController* PC = LP ? LP->PlayerController : nullptr;
+	if (PC)
 	{
-		const FRemoteTalker& Talker = RemoteTalkers[RemoteIndex];
+		// For each registered remote talker
+		for (int32 RemoteIndex = 0; RemoteIndex < RemoteTalkers.Num(); RemoteIndex++)
+		{
+			const FRemoteTalker& Talker = RemoteTalkers[RemoteIndex];
 
-		FUniqueNetIdRepl UniqueIdRepl(Talker.TalkerId);
+			FUniqueNetIdRepl UniqueIdRepl(Talker.TalkerId);
 
-		// Is the remote talker on this local player's mute list?
-		if (SystemMuteList.Find(FUniqueNetIdWrapper(Talker.TalkerId->AsShared())) == INDEX_NONE)
-		{
-			// Unmute on the server
-			PlayerController->ServerUnmutePlayer(UniqueIdRepl);
-		}
-		else
-		{
-			// Mute on the server
-			PlayerController->ServerMutePlayer(UniqueIdRepl);
-		}
+			// Is the remote talker on this local player's mute list?
+			if (SystemMuteList.Find(FUniqueNetIdWrapper(Talker.TalkerId->AsShared())) == INDEX_NONE)
+			{
+				// Unmute on the server
+				PC->ServerUnmutePlayer(UniqueIdRepl);
+			}
+			else
+			{
+				// Mute on the server
+				PC->ServerMutePlayer(UniqueIdRepl);
+			}
 
-		// The ServerUn/MutePlayer() functions will perform the muting based
-		// upon gameplay settings and other player's mute list
+			// The ServerUn/MutePlayer() functions will perform the muting based
+			// upon gameplay settings and other player's mute list
+		}
 	}
 }
 
diff --git a/Engine/Plugins/Online/OnlineSubsystemUtils/Source/OnlineSubsystemUtils/Public/VoiceInterfaceImpl.h b/Engine/Plugins/Online/OnlineSubsystemUtils/Source/OnlineSubsystemUtils/Public/VoiceInterfaceImpl.h
--- a/Engine/Plugins/Online/OnlineSubsystemUtils/Source/OnlineSubsystemUtils/Public/VoiceInterfaceImpl.h
+++ b/Engine/Plugins/Online/OnlineSubsystemUtils/Source/OnlineSubsystemUtils/Public/VoiceInterfaceImpl.h
@@ -119,7 +119,7 @@ PACKAGE_SCOPE:
 	 * @param TalkerIndex the talker that needs the mute list checked for
 	 * @param PlayerController the player controller associated with this talker
 	 */
-	virtual void UpdateMuteListForLocalTalker(int32 TalkerIndex, class APlayerController* PlayerController);
+	virtual void UpdateMuteListForLocalTalker(int32 TalkerIndex);
 
 public:
 
