Make servers send close-reasons/NMT_CloseReason to client, for log-info/debugging. Also fixed NMT_Failure early disconnects from missing NMT_CloseReason - and added more NetDriver/Beacon close reason coverage.

diff --git a/Engine/Plugins/Online/OnlineSubsystemUtils/Source/OnlineSubsystemUtils/Private/OnlineBeaconHost.cpp b/Engine/Plugins/Online/OnlineSubsystemUtils/Source/OnlineSubsystemUtils/Private/OnlineBeaconHost.cpp
--- a/Engine/Plugins/Online/OnlineSubsystemUtils/Source/OnlineSubsystemUtils/Private/OnlineBeaconHost.cpp
+++ b/Engine/Plugins/Online/OnlineSubsystemUtils/Source/OnlineSubsystemUtils/Private/OnlineBeaconHost.cpp
@@ -139,7 +139,9 @@ bool AOnlineBeaconHost::HandleControlMessage(UNetConnection* Connection, uint8 M
 		// only setup state for the hello message
 		if (MessageType != NMT_Hello || Connection->OwningActor != nullptr)
 		{
-			SendFailurePacket(Connection, FText::Format(Error_ControlFlow, FText::FromString(TEXT("NMT_Hello"))));
+			SendFailurePacket(Connection, ENetCloseResult::BeaconControlFlowError,
+								FText::Format(Error_ControlFlow, FText::FromString(TEXT("NMT_Hello"))));
+
 			return false;
 		}
 
@@ -155,7 +157,9 @@ bool AOnlineBeaconHost::HandleControlMessage(UNetConnection* Connection, uint8 M
 		{
 			if (ConnState->bHasSentHello)
 			{
-				SendFailurePacket(Connection, FText::Format(Error_ControlFlow, FText::FromString(TEXT("NMT_Hello"))));
+				SendFailurePacket(Connection, ENetCloseResult::BeaconControlFlowError,
+									FText::Format(Error_ControlFlow, FText::FromString(TEXT("NMT_Hello"))));
+
 				return false;
 			}
 			ConnState->bHasSentHello = true;
@@ -167,7 +171,9 @@ bool AOnlineBeaconHost::HandleControlMessage(UNetConnection* Connection, uint8 M
 			FString EncryptionToken;
 			if (!FNetControlMessage<NMT_Hello>::Receive(Bunch, IsLittleEndian, RemoteNetworkVersion, EncryptionToken))
 			{
-				SendFailurePacket(Connection, FText::Format(Error_UnableToParsePacket, FText::FromString(TEXT("NMT_Hello"))));
+				SendFailurePacket(Connection, ENetCloseResult::BeaconUnableToParsePacket,
+									FText::Format(Error_UnableToParsePacket, FText::FromString(TEXT("NMT_Hello"))));
+
 				return false;
 			}
 
@@ -191,7 +197,7 @@ bool AOnlineBeaconHost::HandleControlMessage(UNetConnection* Connection, uint8 M
 				if (!FNetDelegates::OnReceivedNetworkEncryptionToken.IsBound())
 				{
 					static const FText ErrorTxt = NSLOCTEXT("NetworkErrors", "BeaconEncryptionNotBoundError", "Encryption failure");
-					SendFailurePacket(Connection, ErrorTxt);
+					SendFailurePacket(Connection, ENetCloseResult::EncryptionFailure, ErrorTxt);
 					return false;
 				}
 
@@ -204,7 +210,9 @@ bool AOnlineBeaconHost::HandleControlMessage(UNetConnection* Connection, uint8 M
 		{
 			if (!ConnState->bHasSentChallenge || ConnState->bHasSentLogin)
 			{
-				SendFailurePacket(Connection, FText::Format(Error_ControlFlow, FText::FromString(TEXT("NMT_Login"))));
+				SendFailurePacket(Connection, ENetCloseResult::BeaconControlFlowError,
+									FText::Format(Error_ControlFlow, FText::FromString(TEXT("NMT_Login"))));
+
 				return false;
 			}
 			ConnState->bHasSentLogin = true;
@@ -220,14 +228,16 @@ bool AOnlineBeaconHost::HandleControlMessage(UNetConnection* Connection, uint8 M
 
 			if (!bReceived)
 			{
-				SendFailurePacket(Connection, FText::Format(Error_UnableToParsePacket, FText::FromString(TEXT("NMT_Login"))));
+				SendFailurePacket(Connection, ENetCloseResult::BeaconUnableToParsePacket,
+									FText::Format(Error_UnableToParsePacket, FText::FromString(TEXT("NMT_Login"))));
+
 				return false;
 			}
 
 			if (!UniqueIdRepl.IsValid())
 			{
 				static const FText ErrorTxt = NSLOCTEXT("NetworkErrors", "BeaconLoginInvalidIdError", "Login Failure. Invalid ID.");
-				SendFailurePacket(Connection, ErrorTxt);
+				SendFailurePacket(Connection, ENetCloseResult::BeaconLoginInvalidIdError, ErrorTxt);
 				return false;
 			}
 
@@ -252,7 +262,7 @@ bool AOnlineBeaconHost::HandleControlMessage(UNetConnection* Connection, uint8 M
 			if (!StartVerifyAuthentication(*UniqueIdRepl, UGameplayStatics::ParseOption(OptionsURL, TEXT("AuthTicket"))))
 			{
 				static const FText ErrorTxt = NSLOCTEXT("NetworkErrors", "BeaconLoginInvalidAuthHandlerError", "Login Failure. Unable to process authentication.");
-				SendFailurePacket(Connection, ErrorTxt);
+				SendFailurePacket(Connection, ENetCloseResult::BeaconLoginInvalidAuthHandlerError, ErrorTxt);
 				return false;
 			}
 		}
@@ -261,7 +271,9 @@ bool AOnlineBeaconHost::HandleControlMessage(UNetConnection* Connection, uint8 M
 		{
 			if (!ConnState->bHasSentWelcome || ConnState->bHasSetNetspeed)
 			{
-				SendFailurePacket(Connection, FText::Format(Error_ControlFlow, FText::FromString(TEXT("NMT_Netspeed"))));
+				SendFailurePacket(Connection, ENetCloseResult::BeaconControlFlowError,
+									FText::Format(Error_ControlFlow, FText::FromString(TEXT("NMT_Netspeed"))));
+
 				return false;
 			}
 			ConnState->bHasSetNetspeed = true;
@@ -269,7 +281,9 @@ bool AOnlineBeaconHost::HandleControlMessage(UNetConnection* Connection, uint8 M
 			int32 Rate;
 			if (!FNetControlMessage<NMT_Netspeed>::Receive(Bunch, Rate))
 			{
-				SendFailurePacket(Connection, FText::Format(Error_UnableToParsePacket, FText::FromString(TEXT("NMT_Netspeed"))));
+				SendFailurePacket(Connection, ENetCloseResult::BeaconUnableToParsePacket,
+									FText::Format(Error_UnableToParsePacket, FText::FromString(TEXT("NMT_Netspeed"))));
+
 				return false;
 			}
 			Connection->CurrentNetSpeed = FMath::Clamp(Rate, 1800, NetDriver->MaxClientRate);
@@ -280,7 +294,9 @@ bool AOnlineBeaconHost::HandleControlMessage(UNetConnection* Connection, uint8 M
 		{
 			if (!ConnState->bHasSetNetspeed || ConnState->bHasJoined || (bAuthRequired && !ConnState->bHasAuthenticated))
 			{
-				SendFailurePacket(Connection, FText::Format(Error_ControlFlow, FText::FromString(TEXT("NMT_BeaconJoin"))));
+				SendFailurePacket(Connection, ENetCloseResult::BeaconControlFlowError,
+									FText::Format(Error_ControlFlow, FText::FromString(TEXT("NMT_BeaconJoin"))));
+
 				return false;
 			}
 			ConnState->bHasJoined = true;
@@ -289,7 +305,9 @@ bool AOnlineBeaconHost::HandleControlMessage(UNetConnection* Connection, uint8 M
 			FUniqueNetIdRepl UniqueId;
 			if (!FNetControlMessage<NMT_BeaconJoin>::Receive(Bunch, BeaconType, UniqueId))
 			{
-				SendFailurePacket(Connection, FText::Format(Error_UnableToParsePacket, FText::FromString(TEXT("NMT_BeaconJoin"))));
+				SendFailurePacket(Connection, ENetCloseResult::BeaconUnableToParsePacket,
+									FText::Format(Error_UnableToParsePacket, FText::FromString(TEXT("NMT_BeaconJoin"))));
+
 				return false;
 			}
 
@@ -300,7 +318,11 @@ bool AOnlineBeaconHost::HandleControlMessage(UNetConnection* Connection, uint8 M
 				if (!UniqueId.IsValid() || UniqueId != Connection->PlayerId)
 				{
 					static const FText AuthErrorText = NSLOCTEXT("NetworkErrors", "BeaconAuthError", "Unable to authenticate for beacon. {0} is not valid for connection owned by {1}");
-					SendFailurePacket(Connection, FText::Format(AuthErrorText, FText::FromString(UniqueId.ToDebugString()), FText::FromString(Connection->PlayerId.ToDebugString())));
+
+					SendFailurePacket(Connection, ENetCloseResult::BeaconAuthError,
+										FText::Format(AuthErrorText, FText::FromString(UniqueId.ToDebugString()),
+														FText::FromString(Connection->PlayerId.ToDebugString())));
+
 					return false;
 				}
 				UE_LOG(LogBeacon, Log, TEXT("%s: Beacon Join %s %s"), *GetDebugName(Connection), *BeaconType, *UniqueId.ToDebugString());
@@ -314,14 +336,14 @@ bool AOnlineBeaconHost::HandleControlMessage(UNetConnection* Connection, uint8 M
 			if (Connection->GetClientWorldPackageName() != NAME_None)
 			{
 				static const FText ErrorTxt = NSLOCTEXT("NetworkErrors", "BeaconSpawnClientWorldPackageNameError", "Join failure, existing ClientWorldPackageName.");
-				SendFailurePacket(Connection, ErrorTxt);
+				SendFailurePacket(Connection, ENetCloseResult::BeaconSpawnClientWorldPackageNameError, ErrorTxt);
 				return false;
 			}
 			AOnlineBeaconClient* ClientActor = GetClientActor(Connection);
 			if (ClientActor != nullptr)
 			{
 				static const FText ErrorTxt = NSLOCTEXT("NetworkErrors", "BeaconSpawnExistingActorError", "Join failure, existing beacon actor.");
-				SendFailurePacket(Connection, ErrorTxt);
+				SendFailurePacket(Connection, ENetCloseResult::BeaconSpawnExistingActorError, ErrorTxt);
 				return false;
 			}
 			UWorld* World = GetWorld();
@@ -339,7 +361,7 @@ bool AOnlineBeaconHost::HandleControlMessage(UNetConnection* Connection, uint8 M
 			if (NewClientActor == nullptr || BeaconType != NewClientActor->GetBeaconType())
 			{
 				static const FText ErrorTxt = NSLOCTEXT("NetworkErrors", "BeaconSpawnFailureError", "Join failure, Couldn't spawn client beacon actor.");
-				SendFailurePacket(Connection, ErrorTxt);
+				SendFailurePacket(Connection, ENetCloseResult::BeaconSpawnFailureError, ErrorTxt);
 				return false;
 			}
 
@@ -362,7 +384,9 @@ bool AOnlineBeaconHost::HandleControlMessage(UNetConnection* Connection, uint8 M
 		{
 			if (!ConnState->bHasJoined || ConnState->bHasCompletedAck)
 			{
-				SendFailurePacket(Connection, FText::Format(Error_ControlFlow, FText::FromString(TEXT("NMT_BeaconNetGUIDAck"))));
+				SendFailurePacket(Connection, ENetCloseResult::BeaconControlFlowError,
+									FText::Format(Error_ControlFlow, FText::FromString(TEXT("NMT_BeaconNetGUIDAck"))));
+
 				return false;
 			}
 			ConnState->bHasCompletedAck = true;
@@ -370,7 +394,9 @@ bool AOnlineBeaconHost::HandleControlMessage(UNetConnection* Connection, uint8 M
 			FString BeaconType;
 			if (!FNetControlMessage<NMT_BeaconNetGUIDAck>::Receive(Bunch, BeaconType))
 			{
-				SendFailurePacket(Connection, FText::Format(Error_UnableToParsePacket, FText::FromString(TEXT("NMT_BeaconNetGUIDAck"))));
+				SendFailurePacket(Connection, ENetCloseResult::BeaconUnableToParsePacket,
+									FText::Format(Error_UnableToParsePacket, FText::FromString(TEXT("NMT_BeaconNetGUIDAck"))));
+
 				return false;
 			}
 
@@ -378,7 +404,7 @@ bool AOnlineBeaconHost::HandleControlMessage(UNetConnection* Connection, uint8 M
 			if (ClientActor == nullptr || BeaconType != ClientActor->GetBeaconType())
 			{
 				static const FText ErrorTxt = NSLOCTEXT("NetworkErrors", "BeaconSpawnNetGUIDAckError2", "Join failure, no actor at NetGUIDAck.");
-				SendFailurePacket(Connection, ErrorTxt);
+				SendFailurePacket(Connection, ENetCloseResult::BeaconSpawnNetGUIDAckNoActor, ErrorTxt);
 				return false;
 			}
 
@@ -386,7 +412,7 @@ bool AOnlineBeaconHost::HandleControlMessage(UNetConnection* Connection, uint8 M
 			if (OnBeaconConnectedDelegate == nullptr)
 			{
 				static const FText ErrorTxt = NSLOCTEXT("NetworkErrors", "BeaconSpawnNetGUIDAckError1", "Join failure, no host object at NetGUIDAck.");
-				SendFailurePacket(Connection, ErrorTxt);
+				SendFailurePacket(Connection, ENetCloseResult::BeaconSpawnNetGUIDAckNoHost, ErrorTxt);
 				return false;
 			}
 
@@ -404,20 +430,11 @@ bool AOnlineBeaconHost::HandleControlMessage(UNetConnection* Connection, uint8 M
 			OnBeaconConnectedDelegate->ExecuteIfBound(ClientActor, Connection);
 		}
 		break;
-	case NMT_CloseReason:
-		{
-			FString CloseReasonList;
 
-			if (FNetControlMessage<NMT_CloseReason>::Receive(Bunch, CloseReasonList) && !CloseReasonList.IsEmpty())
-			{
-				Connection->HandleReceiveCloseReason(CloseReasonList);
-			}
-		}
-		break;
 	default:
 		{
 			static const FText ErrorTxt = NSLOCTEXT("NetworkErrors", "BeaconSpawnUnexpectedError", "Join failure, unexpected control message.");
-			SendFailurePacket(Connection, ErrorTxt);
+			SendFailurePacket(Connection, ENetCloseResult::BeaconSpawnUnexpectedError, ErrorTxt);
 			return false;
 		}
 	}
@@ -426,11 +443,27 @@ bool AOnlineBeaconHost::HandleControlMessage(UNetConnection* Connection, uint8 M
 	return true;
 }
 
+void AOnlineBeaconHost::SendFailurePacket(UNetConnection* Connection, FNetCloseResult&& CloseReason, const FText& ErrorText)
+{
+	if (Connection != nullptr)
+	{
+		FString ErrorMsg = ErrorText.ToString();
+
+		UE_LOG(LogBeacon, Log, TEXT("%s: Send failure: %s"), ToCStr(GetDebugName(Connection)), ToCStr(ErrorMsg));
+
+		Connection->SendCloseReason(MoveTemp(CloseReason));
+		FNetControlMessage<NMT_Failure>::Send(Connection, ErrorMsg);
+	}
+}
+
 void AOnlineBeaconHost::SendFailurePacket(UNetConnection* Connection, const FText& ErrorText)
 {
-	FString ErrorMsg = ErrorText.ToString();
-	UE_LOG(LogBeacon, Log, TEXT("%s: Send failure: %s"), *GetDebugName(Connection), *ErrorMsg);
-	FNetControlMessage<NMT_Failure>::Send(Connection, ErrorMsg);
+	if (Connection != nullptr)
+	{
+		FString ErrorMsg = ErrorText.ToString();
+		UE_LOG(LogBeacon, Log, TEXT("%s: Send failure: %s"), *GetDebugName(Connection), *ErrorMsg);
+		FNetControlMessage<NMT_Failure>::Send(Connection, ErrorMsg);
+	}
 }
 
 void AOnlineBeaconHost::CloseHandshakeConnection(UNetConnection* Connection)
@@ -523,7 +556,7 @@ void AOnlineBeaconHost::OnAuthenticationVerificationComplete(const class FUnique
 			{
 				// Auth failure.
 				UE_LOG(LogBeacon, Log, TEXT("%s: Failed to authenticate user: %s error: %s"), *GetDebugName(Connection), *PlayerId.ToString(), *Error.ToLogString());
-				SendFailurePacket(Connection, Error_Authentication);
+				SendFailurePacket(Connection, ENetCloseResult::BeaconAuthenticationFailure, Error_Authentication);
 				CloseHandshakeConnection(Connection);
 			}
 		}
@@ -594,10 +627,12 @@ void AOnlineBeaconHost::OnEncryptionResponse(const FEncryptionKeyResponse& Respo
 			{
 				FString ResponseStr(LexToString(Response.Response));
 				UE_LOG(LogBeacon, Warning, TEXT("%s: OnlineBeaconHost::OnEncryptionResponse: encryption failure [%s] %s"), *GetDebugName(Connection), *ResponseStr, *Response.ErrorMsg);
+
+				Connection->SendCloseReason(ENetCloseResult::EncryptionFailure);
 				FNetControlMessage<NMT_Failure>::Send(Connection, ResponseStr);
 				Connection->FlushNet();
 				// Can't close the connection here since it will leave the failure message in the send buffer and just close the socket. 
-				// Connection->Close();
+				// Connection->Close(ENetCloseResult::EncryptionFailure);
 			}
 		}
 		else
diff --git a/Engine/Plugins/Online/OnlineSubsystemUtils/Source/OnlineSubsystemUtils/Public/OnlineBeaconHost.h b/Engine/Plugins/Online/OnlineSubsystemUtils/Source/OnlineSubsystemUtils/Public/OnlineBeaconHost.h
--- a/Engine/Plugins/Online/OnlineSubsystemUtils/Source/OnlineSubsystemUtils/Public/OnlineBeaconHost.h
+++ b/Engine/Plugins/Online/OnlineSubsystemUtils/Source/OnlineSubsystemUtils/Public/OnlineBeaconHost.h
@@ -5,6 +5,7 @@
 #include "CoreMinimal.h"
 #include "UObject/ObjectMacros.h"
 #include "OnlineBeacon.h"
+#include "Net/Core/Connection/NetCloseResult.h"
 #include "OnlineBeaconHost.generated.h"
 
 class AOnlineBeaconClient;
@@ -24,6 +25,8 @@ struct FOnlineError;
 UCLASS(transient, notplaceable, config=Engine)
 class ONLINESUBSYSTEMUTILS_API AOnlineBeaconHost : public AOnlineBeacon
 {
+	using FNetCloseResult = UE::Net::FNetCloseResult;
+
 	GENERATED_UCLASS_BODY()
 
 public:
@@ -157,7 +160,11 @@ private:
 	void OnConnectionClosed(UNetConnection* Connection);
 
 	bool HandleControlMessage(UNetConnection* Connection, uint8 MessageType, FInBunch& Bunch);
+	void SendFailurePacket(UNetConnection* Connection, FNetCloseResult&& CloseReason, const FText& ErrorText);
+
+	UE_DEPRECATED(5.1, "SendFailurePacket without CloseReason is deprecated. Use the version which takes CloseReason.")
 	void SendFailurePacket(UNetConnection* Connection, const FText& ErrorMsg);
+
 	void CloseHandshakeConnection(UNetConnection* Connection);
 
 	bool GetConnectionDataForUniqueNetId(const FUniqueNetId& UniqueNetId, UNetConnection*& OutConnection, FConnectionState*& OutConnectionState);
diff --git a/Engine/Source/Runtime/Engine/Classes/Engine/NetConnection.h b/Engine/Source/Runtime/Engine/Classes/Engine/NetConnection.h
--- a/Engine/Source/Runtime/Engine/Classes/Engine/NetConnection.h
+++ b/Engine/Source/Runtime/Engine/Classes/Engine/NetConnection.h
@@ -1693,6 +1693,25 @@ public:
 		return NetPing.Get();
 	}
 
+	/**
+	 * Sends an NMT_CloseReason message, with the specified close reason or close reason chain.
+	 * Called automatically by UNetConnection::Close, but should be called beforehand if sending a packet that will trigger a remote close.
+	 *
+	 * @param CloseReason	The close reason or close reason chain to send
+	 */
+	void SendCloseReason(FNetCloseResult&& CloseReason)
+	{
+		SendCloseReason(static_cast<FNetResult&&>(MoveTemp(CloseReason)));
+	}
+
+	/**
+	 * Sends an NMT_CloseReason message, with the specified close reason or close reason chain.
+	 * Called automatically by UNetConnection::Close, but should be called beforehand if sending a packet that will trigger a remote close.
+	 *
+	 * @param CloseReason	The close reason or close reason chain to send
+	 */
+	ENGINE_API void SendCloseReason(FNetResult&& CloseReason);
+
 	/**
 	 * Handles parsing/validation and logging of NMT_CloseReason messages
 	 *
diff --git a/Engine/Source/Runtime/Engine/Private/DataChannel.cpp b/Engine/Source/Runtime/Engine/Private/DataChannel.cpp
--- a/Engine/Source/Runtime/Engine/Private/DataChannel.cpp
+++ b/Engine/Source/Runtime/Engine/Private/DataChannel.cpp
@@ -1306,6 +1306,8 @@ FPacketIdRange UChannel::SendBunch( FOutBunch* Bunch, bool Merge )
 
 		// Bail out, we can't recover from this (without increasing RELIABLE_BUFFER)
 		FString ErrorMsg = NSLOCTEXT("NetworkErrors", "ClientReliableBufferOverflow", "Outgoing reliable buffer overflow").ToString();
+
+		Connection->SendCloseReason(ENetCloseResult::ReliableBufferOverflow);
 		FNetControlMessage<NMT_Failure>::Send(Connection, ErrorMsg);
 		Connection->FlushNet(true);
 		Connection->Close(ENetCloseResult::ReliableBufferOverflow);
@@ -1766,6 +1768,17 @@ void UControlChannel::ReceivedBunch( FInBunch& Bunch )
 		{
 			ReceiveDestructionInfo(Bunch);
 		}
+		else if (MessageType == NMT_CloseReason)
+		{
+			FString CloseReasonList;
+
+			if (FNetControlMessage<NMT_CloseReason>::Receive(Bunch, CloseReasonList) && !CloseReasonList.IsEmpty())
+			{
+				Connection->HandleReceiveCloseReason(CloseReasonList);
+			}
+
+			break;
+		}
 		else if (MessageType == NMT_NetPing)
 		{
 			ENetPingControlMessage NetPingMessageType;
diff --git a/Engine/Source/Runtime/Engine/Private/NetConnection.cpp b/Engine/Source/Runtime/Engine/Private/NetConnection.cpp
--- a/Engine/Source/Runtime/Engine/Private/NetConnection.cpp
+++ b/Engine/Source/Runtime/Engine/Private/NetConnection.cpp
@@ -924,8 +924,6 @@ const TCHAR* LexToString(const EConnectionState Value)
 
 void UNetConnection::Close(FNetResult&& CloseReason)
 {
-	using namespace UE::Net;
-
 	if (IsInternalAck())
 	{
 		SetReserveDestroyedChannels(false);
@@ -939,40 +937,17 @@ void UNetConnection::Close(FNetResult&& CloseReason)
 		UE_LOG(LogNet, Log, TEXT("UNetConnection::Close: %s, Channels: %i, Time: %s"), *Describe(), OpenChannels.Num(),
 				*FDateTime::UtcNow().ToString(TEXT("%Y.%m.%d-%H.%M.%S")));
 
-		FString CloseReasonsStr;
-		FNetCloseResult* CastedCloseReason = Cast<ENetCloseResult>(&CloseReason);
-
-		if (CastedCloseReason == nullptr || *CastedCloseReason != ENetCloseResult::Unknown)
-		{
-			UE_LOG(LogNet, Log, TEXT("UNetConnection::Close: CloseReason:"));
-
-			for (FNetCloseResult::FConstIterator It(CloseReason); It; ++It)
-			{
-				if (CloseReasonsStr.Len() > 0)
-				{
-					CloseReasonsStr += TEXT(',');
-				}
-
-				CloseReasonsStr += It->DynamicToString(ENetResultString::ResultEnumOnly);
-
-				UE_LOG(LogNet, Log, TEXT(" - %s"), ToCStr(It->DynamicToString()));
-			}
-		}
-
-		const bool bReadyToSend = (Handler == nullptr || Handler->IsFullyInitialized()) && HasReceivedClientPacket();
+		SendCloseReason(MoveTemp(CloseReason));
 
 		if (Channels[0] != nullptr)
 		{
-			if (bReadyToSend && Driver->ServerConnection != nullptr && !CloseReasonsStr.IsEmpty())
-			{
-				FNetControlMessage<NMT_CloseReason>::Send(this, CloseReasonsStr);
-			}
-
 			Channels[0]->Close(EChannelCloseReason::Destroyed);
 		}
 
 		SetConnectionState(EConnectionState::USOCK_Closed);
 
+		const bool bReadyToSend = (Handler == nullptr || Handler->IsFullyInitialized()) && HasReceivedClientPacket();
+
 		if (bReadyToSend)
 		{
 			FlushNet();
@@ -980,11 +955,6 @@ void UNetConnection::Close(FNetResult&& CloseReason)
 
 		if (NetAnalyticsData.IsValid())
 		{
-			if (!AnalyticsVars.CloseReason.IsValid())
-			{
-				AnalyticsVars.CloseReason = MakeUnique<FNetResult>(MoveTemp(CloseReason));
-			}
-
 			NetAnalyticsData->CommitAnalytics(AnalyticsVars);
 		}
 
@@ -1014,8 +984,54 @@ void UNetConnection::HandleNetResultOrClose(ENetCloseResult InResult)
 	}
 }
 
+void UNetConnection::SendCloseReason(FNetResult&& CloseReason)
+{
+	using namespace UE::Net;
+
+	FString CloseReasonsStr;
+	FNetCloseResult* CastedCloseReason = Cast<ENetCloseResult>(&CloseReason);
+
+	if (CastedCloseReason == nullptr || *CastedCloseReason != ENetCloseResult::Unknown)
+	{
+		UE_LOG(LogNet, Log, TEXT("UNetConnection::SendCloseReason:"));
+
+		for (FNetCloseResult::FConstIterator It(CloseReason); It; ++It)
+		{
+			if (CloseReasonsStr.Len() > 0)
+			{
+				CloseReasonsStr += TEXT(',');
+			}
+
+			CloseReasonsStr += It->DynamicToString(ENetResultString::ResultEnumOnly);
+
+			UE_LOG(LogNet, Log, TEXT(" - %s"), ToCStr(It->DynamicToString()));
+		}
+	}
+
+	const bool bReadyToSend = (Handler == nullptr || Handler->IsFullyInitialized()) && HasReceivedClientPacket();
+
+	if (Channels.Num() > 0 && Channels[0] != nullptr && bReadyToSend && !CloseReasonsStr.IsEmpty())
+	{
+		FNetControlMessage<NMT_CloseReason>::Send(this, CloseReasonsStr);
+	}
+
+	if (NetAnalyticsData.IsValid())
+	{
+		if (!AnalyticsVars.CloseReason.IsValid())
+		{
+			AnalyticsVars.CloseReason = MakeUnique<FNetResult>(MoveTemp(CloseReason));
+		}
+		// Sometimes SendCloseReason is called separately to Close, both with the same CloseReason, when remote disconnect is anticipated
+		else if (*AnalyticsVars.CloseReason != CloseReason)
+		{
+			AnalyticsVars.CloseReason->AddChainResult(MoveTemp(CloseReason));
+		}
+	}
+}
+
 void UNetConnection::HandleReceiveCloseReason(const FString& CloseReasonList)
 {
+	// Multiple NMT_CloseReason's can be received in some circumstances, but only process the first to limit spam potential.
 	if (!bReceivedCloseReason)
 	{
 		bool bValid = true;
@@ -1036,23 +1052,26 @@ void UNetConnection::HandleReceiveCloseReason(const FString& CloseReasonList)
 
 		if (bValid)
 		{
+			const bool bRemoteIsServer = Driver && Driver->ServerConnection;
 			TArray<FString> CloseReasonsAlloc;
-			TArray<FString>& ClientCloseReasons = (NetAnalyticsData.IsValid() ? AnalyticsVars.ClientCloseReasons : CloseReasonsAlloc);
+			TArray<FString>& CloseReasonsArray = ((!bRemoteIsServer && NetAnalyticsData.IsValid()) ?
+													AnalyticsVars.ClientCloseReasons : CloseReasonsAlloc);
 
-			CloseReasonList.ParseIntoArray(ClientCloseReasons, TEXT(","));
+			CloseReasonList.ParseIntoArray(CloseReasonsArray, TEXT(","));
 
-			if (ClientCloseReasons.Num() < 128)
+			if (CloseReasonsArray.Num() < 128)
 			{
-				UE_LOG(LogNet, Log, TEXT("NMT_CloseReason: (Client Disconnect Reasons) %s"), ToCStr(LowLevelGetRemoteAddress(true)));
+				UE_LOG(LogNet, Log, TEXT("NMT_CloseReason: (%s Disconnect Reasons) %s"), (bRemoteIsServer ? TEXT("Server") : TEXT("Client")),
+						ToCStr(LowLevelGetRemoteAddress(true)));
 
-				for (const FString& CurReason : ClientCloseReasons)
+				for (const FString& CurReason : CloseReasonsArray)
 				{
 					UE_LOG(LogNet, Log, TEXT(" - %s"), ToCStr(CurReason));
 				}
 			}
 			else
 			{
-				ClientCloseReasons.Empty();
+				CloseReasonsArray.Empty();
 			}
 
 			bReceivedCloseReason = true;
@@ -4937,10 +4956,12 @@ void UNetConnection::SendChallengeControlMessage(const FEncryptionKeyResponse& R
 		{
 			FString ResponseStr(LexToString(Response.Response));
 			UE_LOG(LogNet, Warning, TEXT("UWorld::SendChallengeControlMessage: encryption failure [%s] %s"), *ResponseStr, *Response.ErrorMsg);
+
+			SendCloseReason(ENetCloseResult::EncryptionFailure);
 			FNetControlMessage<NMT_Failure>::Send(this, ResponseStr);
 			FlushNet();
 			// Can't close the connection here since it will leave the failure message in the send buffer and just close the socket. 
-			// Close();
+			// Close(ENetCloseResult::EncryptionFailure);
 		}
 	}
 	else
diff --git a/Engine/Source/Runtime/Engine/Private/NetDriver.cpp b/Engine/Source/Runtime/Engine/Private/NetDriver.cpp
--- a/Engine/Source/Runtime/Engine/Private/NetDriver.cpp
+++ b/Engine/Source/Runtime/Engine/Private/NetDriver.cpp
@@ -1795,11 +1795,13 @@ void UNetDriver::Shutdown()
 	// Server closing connections with clients
 	if (ClientConnections.Num() > 0)
 	{
-		for (int32 ClientIndex = 0; ClientIndex < ClientConnections.Num(); ClientIndex++)
+		FString ErrorMsg = NSLOCTEXT("NetworkErrors", "HostClosedConnection", "Host closed the connection.").ToString();
+
+		for (UNetConnection* CurClient : ClientConnections)
 		{
-			FString ErrorMsg = NSLOCTEXT("NetworkErrors", "HostClosedConnection", "Host closed the connection.").ToString();
-			FNetControlMessage<NMT_Failure>::Send(ClientConnections[ClientIndex], ErrorMsg);
-			ClientConnections[ClientIndex]->FlushNet(true);
+			CurClient->SendCloseReason(ENetCloseResult::HostClosedConnection);
+			FNetControlMessage<NMT_Failure>::Send(CurClient, ErrorMsg);
+			CurClient->FlushNet(true);
 		}
 
 		for (int32 ClientIndex = ClientConnections.Num() - 1; ClientIndex >= 0; ClientIndex--)
@@ -2306,9 +2308,11 @@ void UNetDriver::ProcessRemoteFunctionForChannelPrivate(
 			UE_LOG(LogNet, Warning, TEXT("Closing connection. Can't send function '%s' on '%s': Reliable buffer overflow. FieldCache->FieldNetIndex: %d Max %d. Ch MaxPacket: %d."), *GetNameSafe(Function), *GetFullNameSafe(TargetObj), FieldCache->FieldNetIndex, ClassCache->GetMaxIndex(), Ch->Connection->MaxPacket );
 
 			FString ErrorMsg = NSLOCTEXT("NetworkErrors", "ClientReliableBufferOverflow", "Outgoing reliable buffer overflow").ToString();
+
+			Connection->SendCloseReason(ENetCloseResult::RPCReliableBufferOverflow);
 			FNetControlMessage<NMT_Failure>::Send(Connection, ErrorMsg);
 			Connection->FlushNet(true);
-			Connection->Close();
+			Connection->Close(ENetCloseResult::RPCReliableBufferOverflow);
 #if USE_SERVER_PERF_COUNTERS
 			PerfCountersIncrement(TEXT("ClosedConnectionsDueToReliableBufferOverflow"));
 #endif
@@ -2919,6 +2923,7 @@ bool UNetDriver::HandleNetDisconnectCommand( const TCHAR* Cmd, FOutputDevice& Ar
 		UE_LOG(LogNet, Log, TEXT("%s disconnecting connection from host [%s]"), 
 			*GetDescription(),*ServerConnection->LowLevelDescribe());
 
+		ServerConnection->SendCloseReason(ENetCloseResult::Disconnect);
 		FNetControlMessage<NMT_Failure>::Send(ServerConnection, Msg);
 	}
 #if WITH_SERVER_CODE
@@ -2932,6 +2937,7 @@ bool UNetDriver::HandleNetDisconnectCommand( const TCHAR* Cmd, FOutputDevice& Ar
 				UE_LOG(LogNet, Log, TEXT("%s disconnecting from client [%s]"), 
 					*GetDescription(),*Connection->LowLevelDescribe());
 
+				Connection->SendCloseReason(ENetCloseResult::Disconnect);
 				FNetControlMessage<NMT_Failure>::Send(Connection, Msg);
 				Connection->FlushNet(true);
 			}
diff --git a/Engine/Source/Runtime/Engine/Private/PendingNetGame.cpp b/Engine/Source/Runtime/Engine/Private/PendingNetGame.cpp
--- a/Engine/Source/Runtime/Engine/Private/PendingNetGame.cpp
+++ b/Engine/Source/Runtime/Engine/Private/PendingNetGame.cpp
@@ -220,7 +220,7 @@ void UPendingNetGame::NotifyControlMessage(UNetConnection* Connection, uint8 Mes
 					Connection->OwningActor ? *Connection->OwningActor->GetName() : TEXT("No Owner"),
 					*ConnectionError);
 
-				Connection->Close();
+				Connection->Close(ENetCloseResult::FailureReceived);
 			}
 
 			break;
diff --git a/Engine/Source/Runtime/Engine/Private/World.cpp b/Engine/Source/Runtime/Engine/Private/World.cpp
--- a/Engine/Source/Runtime/Engine/Private/World.cpp
+++ b/Engine/Source/Runtime/Engine/Private/World.cpp
@@ -5815,10 +5815,16 @@ void UWorld::NotifyControlMessage(UNetConnection* Connection, uint8 MessageType,
 						ErrorMsg = NSLOCTEXT("NetworkErrors", "GenericConnectionFailed", "Connection Failed.").ToString();
 					}
 
+					if (Connection != nullptr)
+					{
+						Connection->SendCloseReason(ENetCloseResult::FailureReceived);
+					}
+
 					GEngine->BroadcastNetworkFailure(this, NetDriver, ENetworkFailure::FailureReceived, ErrorMsg);
-					if (Connection)
+
+					if (Connection != nullptr)
 					{
-						Connection->Close();
+						Connection->Close(ENetCloseResult::FailureReceived);
 					}
 				}
 
@@ -5884,7 +5890,7 @@ void UWorld::NotifyControlMessage(UNetConnection* Connection, uint8 MessageType,
 						UE_LOG(LogNet, Log, TEXT("NotifyControlMessage: Client connecting with invalid version. LocalNetworkVersion: %i, RemoteNetworkVersion: %i"), LocalNetworkVersion, RemoteNetworkVersion);
 						FNetControlMessage<NMT_Upgrade>::Send(Connection, LocalNetworkVersion);
 						Connection->FlushNet(true);
-						Connection->Close();
+						Connection->Close(ENetCloseResult::Upgrade);
 #if USE_SERVER_PERF_COUNTERS
 						PerfCountersIncrement(TEXT("ClosedConnectionsDueToIncompatibleVersion"));
 #endif
@@ -5904,7 +5910,11 @@ void UWorld::NotifyControlMessage(UNetConnection* Connection, uint8 MessageType,
 							else
 							{
 								FString FailureMsg(TEXT("Encryption failure"));
-								UE_LOG(LogNet, Warning, TEXT("%s: No delegate available to handle encryption token, disconnecting."), *Connection->GetName());
+
+								UE_LOG(LogNet, Warning, TEXT("%s: No delegate available to handle encryption token, disconnecting."),
+										ToCStr(Connection->GetName()));
+
+								Connection->SendCloseReason(ENetCloseResult::EncryptionFailure);
 								FNetControlMessage<NMT_Failure>::Send(Connection, FailureMsg);
 								Connection->FlushNet(true);
 							}
@@ -6005,11 +6015,13 @@ void UWorld::NotifyControlMessage(UNetConnection* Connection, uint8 MessageType,
 					{
 						UE_LOG(LogNet, Log, TEXT("PreLogin failure: %s"), *ErrorMsg);
 						NETWORK_PROFILER(GNetworkProfiler.TrackEvent(TEXT("PRELOGIN FAILURE"), *ErrorMsg, Connection));
+
+						Connection->SendCloseReason(ENetCloseResult::PreLoginFailure);
 						FNetControlMessage<NMT_Failure>::Send(Connection, ErrorMsg);
 						Connection->FlushNet(true);
 						//@todo sz - can't close the connection here since it will leave the failure message 
 						// in the send buffer and just close the socket. 
-						//Connection->Close();
+						//Connection->Close(ENetCloseResult::PreLoginFailure);
 					}
 					else
 					{
@@ -6047,11 +6059,13 @@ void UWorld::NotifyControlMessage(UNetConnection* Connection, uint8 MessageType,
 						// Failed to connect.
 						UE_LOG(LogNet, Log, TEXT("Join failure: %s"), *ErrorMsg);
 						NETWORK_PROFILER(GNetworkProfiler.TrackEvent(TEXT("JOIN FAILURE"), *ErrorMsg, Connection));
+
+						Connection->SendCloseReason(ENetCloseResult::JoinFailure);
 						FNetControlMessage<NMT_Failure>::Send(Connection, ErrorMsg);
 						Connection->FlushNet(true);
 						//@todo sz - can't close the connection here since it will leave the failure message 
 						// in the send buffer and just close the socket. 
-						//Connection->Close();
+						//Connection->Close(ENetCloseResult::JoinFailure);
 					}
 					else
 					{
@@ -6145,11 +6159,13 @@ void UWorld::NotifyControlMessage(UNetConnection* Connection, uint8 MessageType,
 						// if any splitscreen viewport fails to join, all viewports on that client also fail
 						UE_LOG(LogNet, Log, TEXT("PreLogin failure: %s"), *ErrorMsg);
 						NETWORK_PROFILER(GNetworkProfiler.TrackEvent(TEXT("PRELOGIN FAILURE"), *ErrorMsg, Connection));
+
+						Connection->SendCloseReason(ENetCloseResult::PreLoginFailure);
 						FNetControlMessage<NMT_Failure>::Send(Connection, ErrorMsg);
 						Connection->FlushNet(true);
 						//@todo sz - can't close the connection here since it will leave the failure message 
 						// in the send buffer and just close the socket. 
-						//Connection->Close();
+						//Connection->Close(ENetCloseResult::PreLoginFailure);
 					}
 					else
 					{
@@ -6179,12 +6195,14 @@ void UWorld::NotifyControlMessage(UNetConnection* Connection, uint8 MessageType,
 							NETWORK_PROFILER(GNetworkProfiler.TrackEvent(TEXT("JOINSPLIT FAILURE"), *ErrorMsg, Connection));
 							// remove the child connection
 							Connection->Children.Remove(ChildConn);
+
 							// if any splitscreen viewport fails to join, all viewports on that client also fail
+							Connection->SendCloseReason(ENetCloseResult::JoinSplitFailure);
 							FNetControlMessage<NMT_Failure>::Send(Connection, ErrorMsg);
 							Connection->FlushNet(true);
 							//@todo sz - can't close the connection here since it will leave the failure message 
 							// in the send buffer and just close the socket. 
-							//Connection->Close();
+							//Connection->Close(ENetCloseResult::JoinSplitFailure);
 						}
 						else
 						{
@@ -6239,18 +6257,6 @@ void UWorld::NotifyControlMessage(UNetConnection* Connection, uint8 MessageType,
 
 				break;
 			}
-
-			case NMT_CloseReason:
-			{
-				FString CloseReasonList;
-
-				if (FNetControlMessage<NMT_CloseReason>::Receive(Bunch, CloseReasonList) && !CloseReasonList.IsEmpty())
-				{
-					Connection->HandleReceiveCloseReason(CloseReasonList);
-				}
-
-				break;
-			}
 		}
 	}
 }
diff --git a/Engine/Source/Runtime/Net/Core/Private/Net/Core/Connection/NetCloseResult.cpp b/Engine/Source/Runtime/Net/Core/Private/Net/Core/Connection/NetCloseResult.cpp
--- a/Engine/Source/Runtime/Net/Core/Private/Net/Core/Connection/NetCloseResult.cpp
+++ b/Engine/Source/Runtime/Net/Core/Private/Net/Core/Connection/NetCloseResult.cpp
@@ -67,6 +67,24 @@ const TCHAR* LexToString(ENetCloseResult InResult)
 	case ENetCloseResult::Extended:
 		return TEXT("Extended");
 
+	case ENetCloseResult::HostClosedConnection:
+		return TEXT("HostClosedConnection");
+
+	case ENetCloseResult::Disconnect:
+		return TEXT("Disconnect");
+
+	case ENetCloseResult::Upgrade:
+		return TEXT("Upgrade");
+
+	case ENetCloseResult::PreLoginFailure:
+		return TEXT("PreLoginFailure");
+
+	case ENetCloseResult::JoinFailure:
+		return TEXT("JoinFailure");
+
+	case ENetCloseResult::JoinSplitFailure:
+		return TEXT("JoinSplitFailure");
+
 	case ENetCloseResult::RPCDoS:
 		return TEXT("RPCDoS");
 
@@ -139,6 +157,9 @@ const TCHAR* LexToString(ENetCloseResult InResult)
 	case ENetCloseResult::LogLimitSustained:
 		return TEXT("LogLimitSustained");
 
+	case ENetCloseResult::EncryptionFailure:
+		return TEXT("EncryptionFailure");
+
 	case ENetCloseResult::ReceivedNetGUIDBunchFail:
 		return TEXT("ReceivedNetGUIDBunchFail");
 
@@ -178,6 +199,9 @@ const TCHAR* LexToString(ENetCloseResult InResult)
 	case ENetCloseResult::ReliableBufferOverflow:
 		return TEXT("ReliableBufferOverflow");
 
+	case ENetCloseResult::RPCReliableBufferOverflow:
+		return TEXT("RPCReliableBufferOverflow");
+
 	case ENetCloseResult::ControlChannelClose:
 		return TEXT("ControlChannelClose");
 
@@ -265,6 +289,42 @@ const TCHAR* LexToString(ENetCloseResult InResult)
 	case ENetCloseResult::FieldPayloadFail:
 		return TEXT("FieldPayloadFail");
 
+	case ENetCloseResult::BeaconControlFlowError:
+		return TEXT("BeaconControlFlowError");
+
+	case ENetCloseResult::BeaconAuthenticationFailure:
+		return TEXT("BeaconAuthenticationFailure");
+
+	case ENetCloseResult::BeaconUnableToParsePacket:
+		return TEXT("BeaconUnableToParsePacket");
+
+	case ENetCloseResult::BeaconLoginInvalidIdError:
+		return TEXT("BeaconLoginInvalidIdError");
+
+	case ENetCloseResult::BeaconLoginInvalidAuthHandlerError:
+		return TEXT("BeaconLoginInvalidAuthHandlerError");
+
+	case ENetCloseResult::BeaconAuthError:
+		return TEXT("BeaconAuthError");
+
+	case ENetCloseResult::BeaconSpawnClientWorldPackageNameError:
+		return TEXT("BeaconSpawnClientWorldPackageNameError");
+
+	case ENetCloseResult::BeaconSpawnExistingActorError:
+		return TEXT("BeaconSpawnExistingActorError");
+
+	case ENetCloseResult::BeaconSpawnFailureError:
+		return TEXT("BeaconSpawnFailureError");
+
+	case ENetCloseResult::BeaconSpawnNetGUIDAckNoActor:
+		return TEXT("BeaconSpawnNetGUIDAckNoActor");
+
+	case ENetCloseResult::BeaconSpawnNetGUIDAckNoHost:
+		return TEXT("BeaconSpawnNetGUIDAckNoHost");
+
+	case ENetCloseResult::BeaconSpawnUnexpectedError:
+		return TEXT("BeaconSpawnUnexpectedError");
+
 	case ENetCloseResult::FaultDisconnect:
 		return TEXT("FaultDisconnect");
 
diff --git a/Engine/Source/Runtime/Net/Core/Public/Net/Core/Connection/NetCloseResult.h b/Engine/Source/Runtime/Net/Core/Public/Net/Core/Connection/NetCloseResult.h
--- a/Engine/Source/Runtime/Net/Core/Public/Net/Core/Connection/NetCloseResult.h
+++ b/Engine/Source/Runtime/Net/Core/Public/Net/Core/Connection/NetCloseResult.h
@@ -79,6 +79,27 @@ enum class ENetCloseResult : uint32
 	Extended,
 
 
+	/** NetDriver Error result types */
+
+	/** Host closed the connection */
+	HostClosedConnection,
+
+	/** Disconnected remotely */
+	Disconnect,
+
+	/** Client needs to upgrade */
+	Upgrade,
+
+	/** PreLogin failed */
+	PreLoginFailure,
+
+	/** Join failed */
+	JoinFailure,
+
+	/** JoinSplit failed */
+	JoinSplitFailure,
+
+
 	/** NetConnection Error result types */
 
 	/** RPC DoS Detection kicked the player */
@@ -153,6 +174,9 @@ enum class ENetCloseResult : uint32
 	/** Log limiting repeated threshold hits kicked the player */
 	LogLimitSustained,
 
+	/** Encryption failure */
+	EncryptionFailure,
+
 
 	/** Channel Error result types */
 
@@ -195,6 +219,9 @@ enum class ENetCloseResult : uint32
 	/** Reliable buffer overflowed when attempting to send */
 	ReliableBufferOverflow,
 
+	/** Reliable buffer overflowed when attempting to send RPC */
+	RPCReliableBufferOverflow,
+
 
 	/** Control Channel result types */
 
@@ -289,6 +316,45 @@ enum class ENetCloseResult : uint32
 	FieldPayloadFail,
 
 
+	/** Beacon Error result types */
+
+	/** Control flow error */
+	BeaconControlFlowError,
+
+	/** Unable to parse expected packet structure */
+	BeaconUnableToParsePacket,
+
+	/** Failed to verify user authentication */
+	BeaconAuthenticationFailure,
+
+	/** Login Failure, invalid ID */
+	BeaconLoginInvalidIdError,
+
+	/** Login Failure, unable to process authentication */
+	BeaconLoginInvalidAuthHandlerError,
+
+	/** Unable to authenticate for beacon, wrong PlayerId */
+	BeaconAuthError,
+
+	/** Join failure, existing ClientWorldPackageName */
+	BeaconSpawnClientWorldPackageNameError,
+
+	/** Join failure, existing beacon actor */
+	BeaconSpawnExistingActorError,
+
+	/** Join failure, couldn't spawn client beacon actor */
+	BeaconSpawnFailureError,
+
+	/** Join failure, no actor at NetGUIDAck */
+	BeaconSpawnNetGUIDAckNoActor,
+
+	/** Join failure, no host object at NetGUIDAck */
+	BeaconSpawnNetGUIDAckNoHost,
+
+	/** Join failure, unexpected control message */
+	BeaconSpawnUnexpectedError,
+
+
 	/** Fault Handler Error result types */
 
 	/** Net Fault Recovery failed to recover from a fault, and is triggering a disconnect */
