Fix OnClientConnected possibly accessing stale pointers.
I've managed to catch a crash caused by this in the debugger and the client beacon was disconned and net connection had been GC'd.

diff --git a/Engine/Plugins/Online/OnlineSubsystemUtils/Source/OnlineSubsystemUtils/Private/OnlineBeaconHostObject.cpp b/Engine/Plugins/Online/OnlineSubsystemUtils/Source/OnlineSubsystemUtils/Private/OnlineBeaconHostObject.cpp
--- a/Engine/Plugins/Online/OnlineSubsystemUtils/Source/OnlineSubsystemUtils/Private/OnlineBeaconHostObject.cpp
+++ b/Engine/Plugins/Online/OnlineSubsystemUtils/Source/OnlineSubsystemUtils/Private/OnlineBeaconHostObject.cpp
@@ -33,9 +33,15 @@ AOnlineBeaconClient* AOnlineBeaconHostObject::SpawnBeaconActor(UNetConnection* C
 
 void AOnlineBeaconHostObject::OnClientConnected(AOnlineBeaconClient* NewClientActor, UNetConnection* ClientConnection)
 {
+	if (NewClientActor == nullptr)
+	{
+		return;
+	}
+
+	UNetConnection* ActorConnection = NewClientActor->GetNetConnection();
 	UE_LOG(LogBeacon, Verbose, TEXT("OnClientConnected %s from (%s)"),
-		NewClientActor ? *NewClientActor->GetName() : TEXT("NULL"),
-		NewClientActor ? *NewClientActor->GetNetConnection()->LowLevelDescribe() : TEXT("NULL"));
+		*NewClientActor->GetName(),
+		ActorConnection ? *ActorConnection->LowLevelDescribe() : TEXT("NULL"));
 
 	ClientActors.Add(NewClientActor);
 }
