Handle cases where UNetDriver::Notify being set to nullptr were not handled.
Set UNetDriver::Notify to nullptr when beacon can no longer handle messages.
Update ensure condition in AOnlineBeaconClient::OnNetCleanup to be valid.

diff --git a/Engine/Plugins/Experimental/WebSocketNetworking/Source/WebSocketNetworking/Private/WebSocketNetDriver.cpp b/Engine/Plugins/Experimental/WebSocketNetworking/Source/WebSocketNetworking/Private/WebSocketNetDriver.cpp
--- a/Engine/Plugins/Experimental/WebSocketNetworking/Source/WebSocketNetworking/Private/WebSocketNetDriver.cpp
+++ b/Engine/Plugins/Experimental/WebSocketNetworking/Source/WebSocketNetworking/Private/WebSocketNetDriver.cpp
@@ -211,7 +211,7 @@ UWebSocketConnection* UWebSocketNetDriver::GetServerConnection()
 void UWebSocketNetDriver::OnWebSocketClientConnected(INetworkingWebSocket* ClientWebSocket)
 {
 	// Determine if allowing for client/server connections
-	const bool bAcceptingConnection = Notify->NotifyAcceptingConnection() == EAcceptConnection::Accept;
+	const bool bAcceptingConnection = Notify != nullptr && Notify->NotifyAcceptingConnection() == EAcceptConnection::Accept;
 	if (bAcceptingConnection)
 	{
 
diff --git a/Engine/Plugins/Online/OnlineSubsystemOculus/Source/Private/OculusNetDriver.cpp b/Engine/Plugins/Online/OnlineSubsystemOculus/Source/Private/OculusNetDriver.cpp
--- a/Engine/Plugins/Online/OnlineSubsystemOculus/Source/Private/OculusNetDriver.cpp
+++ b/Engine/Plugins/Online/OnlineSubsystemOculus/Source/Private/OculusNetDriver.cpp
@@ -190,11 +190,11 @@ void UOculusNetDriver::TickDispatch(float DeltaTime)
 			bool bPassedChallenge = false;
 			TSharedPtr<StatelessConnectHandlerComponent> StatelessConnect;
 
-			if (!ConnectionlessHandler.IsValid() || !StatelessConnectComponent.IsValid())
+			if (Notify == nullptr || !ConnectionlessHandler.IsValid() || !StatelessConnectComponent.IsValid())
 			{
 				UE_LOG(LogNet, Log,
-					TEXT("Invalid ConnectionlessHandler (%i) or StatelessConnectComponent (%i); can't accept connections."),
-					(int32)(ConnectionlessHandler.IsValid()), (int32)(StatelessConnectComponent.IsValid()));
+					TEXT("Invalid Notify (%i) or ConnectionlessHandler (%i) or StatelessConnectComponent (%i); can't accept connections."),
+					(int32)(Notify != nullptr), (int32)(ConnectionlessHandler.IsValid()), (int32)(StatelessConnectComponent.IsValid()));
 				continue;
 			}
 
@@ -346,7 +346,7 @@ void UOculusNetDriver::OnNewNetworkingPeerRequest(ovrMessageHandle Message, bool
 bool UOculusNetDriver::AddNewClientConnection(ovrID PeerID)
 {
 	// Ignore the peer if not accepting new connections
-	if (Notify->NotifyAcceptingConnection() != EAcceptConnection::Accept)
+	if (Notify == nullptr || Notify->NotifyAcceptingConnection() != EAcceptConnection::Accept)
 	{
 		UE_LOG(LogNet, Warning, TEXT("Not accepting more new connections"));
 		return false;
diff --git a/Engine/Plugins/Online/OnlineSubsystemUtils/Source/OnlineSubsystemUtils/Private/IpNetDriver.cpp b/Engine/Plugins/Online/OnlineSubsystemUtils/Source/OnlineSubsystemUtils/Private/IpNetDriver.cpp
--- a/Engine/Plugins/Online/OnlineSubsystemUtils/Source/OnlineSubsystemUtils/Private/IpNetDriver.cpp
+++ b/Engine/Plugins/Online/OnlineSubsystemUtils/Source/OnlineSubsystemUtils/Private/IpNetDriver.cpp
@@ -1476,7 +1476,7 @@ UNetConnection* UIpNetDriver::ProcessConnectionlessPacket(FReceivedPacketView& P
 	bool bRestartedHandshake = false;
 	bool bIgnorePacket = true;
 
-	if (ConnectionlessHandler.IsValid() && StatelessConnectComponent.IsValid())
+	if (Notify != nullptr && ConnectionlessHandler.IsValid() && StatelessConnectComponent.IsValid())
 	{
 		StatelessConnect = StatelessConnectComponent.Pin();
 
@@ -1585,8 +1585,8 @@ UNetConnection* UIpNetDriver::ProcessConnectionlessPacket(FReceivedPacketView& P
 #endif
 	else
 	{
-		UE_LOG(LogNet, Log, TEXT("Invalid ConnectionlessHandler (%i) or StatelessConnectComponent (%i); can't accept connections."),
-				(int32)(ConnectionlessHandler.IsValid()), (int32)(StatelessConnectComponent.IsValid()));
+		UE_LOG(LogNet, Log, TEXT("Invalid Notify (%i) or ConnectionlessHandler (%i) or StatelessConnectComponent (%i); can't accept connections."),
+			(int32)(Notify != nullptr), (int32)(ConnectionlessHandler.IsValid()), (int32)(StatelessConnectComponent.IsValid()));
 	}
 
 	if (bPassedChallenge)
diff --git a/Engine/Plugins/Online/OnlineSubsystemUtils/Source/OnlineSubsystemUtils/Private/OnlineBeacon.cpp b/Engine/Plugins/Online/OnlineSubsystemUtils/Source/OnlineSubsystemUtils/Private/OnlineBeacon.cpp
--- a/Engine/Plugins/Online/OnlineSubsystemUtils/Source/OnlineSubsystemUtils/Private/OnlineBeacon.cpp
+++ b/Engine/Plugins/Online/OnlineSubsystemUtils/Source/OnlineSubsystemUtils/Private/OnlineBeacon.cpp
@@ -32,19 +32,7 @@ bool AOnlineBeacon::InitBase()
 
 void AOnlineBeacon::EndPlay(const EEndPlayReason::Type EndPlayReason)
 {
-	if (NetDriver)
-	{
-		if (NetDriver->IsInTick())
-		{
-			NetDriver->SetPendingDestruction(true);
-		}
-		else
-		{
-			GEngine->DestroyNamedNetDriver(GetWorld(), NetDriverName);
-		}
-		NetDriver = nullptr;
-	}
-
+	CleanupNetDriver();
 	Super::EndPlay(EndPlayReason);
 }
 
@@ -59,6 +47,12 @@ void AOnlineBeacon::DestroyBeacon()
 	UE_LOG(LogBeacon, Verbose, TEXT("Destroying beacon %s, netdriver %s"), *GetName(), NetDriver ? *NetDriver->GetDescription() : TEXT("NULL"));
 	GEngine->OnNetworkFailure().Remove(HandleNetworkFailureDelegateHandle);
 
+	CleanupNetDriver();
+	Destroy();
+}
+
+void AOnlineBeacon::CleanupNetDriver()
+{
 	if (NetDriver)
 	{
 		if (NetDriver->IsInTick())
@@ -66,13 +60,14 @@ void AOnlineBeacon::DestroyBeacon()
 			NetDriver->SetPendingDestruction(true);
 		}
 		else
-		{ 
+		{
 			GEngine->DestroyNamedNetDriver(GetWorld(), NetDriverName);
 		}
+		// If the net connection is currently in the middle of processing messages it is
+		// possible for the notifier to be fired again if not cleared.
+		NetDriver->Notify = nullptr;
 		NetDriver = nullptr;
 	}
-
-	Destroy();
 }
 
 void AOnlineBeacon::HandleNetworkFailure(UWorld *World, UNetDriver *InNetDriver, ENetworkFailure::Type FailureType, const FString& ErrorString)
@@ -87,19 +82,7 @@ void AOnlineBeacon::HandleNetworkFailure(UWorld *World, UNetDriver *InNetDriver,
 void AOnlineBeacon::OnFailure()
 {
 	GEngine->OnNetworkFailure().Remove(HandleNetworkFailureDelegateHandle);
-	
-	if (NetDriver)
-	{
-		if (NetDriver->IsInTick())
-		{
-			NetDriver->SetPendingDestruction(true);
-		}
-		else
-		{
-			GEngine->DestroyNamedNetDriver(GetWorld(), NetDriverName);
-		}
-		NetDriver = nullptr;
-	}
+	CleanupNetDriver();
 }
 
 void AOnlineBeacon::OnActorChannelOpen(FInBunch& Bunch, UNetConnection* Connection)
diff --git a/Engine/Plugins/Online/OnlineSubsystemUtils/Source/OnlineSubsystemUtils/Private/OnlineBeaconClient.cpp b/Engine/Plugins/Online/OnlineSubsystemUtils/Source/OnlineSubsystemUtils/Private/OnlineBeaconClient.cpp
--- a/Engine/Plugins/Online/OnlineSubsystemUtils/Source/OnlineSubsystemUtils/Private/OnlineBeaconClient.cpp
+++ b/Engine/Plugins/Online/OnlineSubsystemUtils/Source/OnlineSubsystemUtils/Private/OnlineBeaconClient.cpp
@@ -317,7 +317,9 @@
 
 void AOnlineBeaconClient::OnNetCleanup(UNetConnection* Connection)
 {
-	ensure(Connection == BeaconConnection);
+	// During the garbage collection of UNetConnection OnNetCleanup may be triggered.
+	// When this happens BeaconConnection will have already been set to nullptr by GC.
+	ensure(NetDriver == nullptr || Connection == BeaconConnection);
 	SetConnectionState(EBeaconConnectionState::Closed);
 
 	AOnlineBeaconHostObject* BeaconHostObject = GetBeaconOwner();
diff --git a/Engine/Plugins/Online/OnlineSubsystemUtils/Source/OnlineSubsystemUtils/Public/OnlineBeacon.h b/Engine/Plugins/Online/OnlineSubsystemUtils/Source/OnlineSubsystemUtils/Public/OnlineBeacon.h
--- a/Engine/Plugins/Online/OnlineSubsystemUtils/Source/OnlineSubsystemUtils/Public/OnlineBeacon.h
+++ b/Engine/Plugins/Online/OnlineSubsystemUtils/Source/OnlineSubsystemUtils/Public/OnlineBeacon.h
@@ -95,6 +95,8 @@ class ONLINESUBSYSTEMUTILS_API AOnlineBeacon : public AActor, public FNetworkNot
 
 protected:
 
+	void CleanupNetDriver();
+
 	/** Time beacon will wait to establish a connection with the beacon host */
 	UPROPERTY(Config)
 	float BeaconConnectionInitialTimeout;
diff --git a/Engine/Plugins/Runtime/Steam/SteamSockets/Source/SteamSockets/Private/SteamSocketsNetDriver.cpp b/Engine/Plugins/Runtime/Steam/SteamSockets/Source/SteamSockets/Private/SteamSocketsNetDriver.cpp
--- a/Engine/Plugins/Runtime/Steam/SteamSockets/Source/SteamSockets/Private/SteamSocketsNetDriver.cpp
+++ b/Engine/Plugins/Runtime/Steam/SteamSockets/Source/SteamSockets/Private/SteamSocketsNetDriver.cpp
@@ -467,7 +467,7 @@ void USteamSocketsNetDriver::OnConnectionCreated(SteamSocketHandles ListenParent
 
 	// Unlike the other SteamNetworking functionality, connections that we don't want cannot just be ignored
 	// So instead of processing everything and then disconnecting unwanted connections, we drop them immediately.
-	if (Notify->NotifyAcceptingConnection() == EAcceptConnection::Accept)
+	if (Notify != nullptr && Notify->NotifyAcceptingConnection() == EAcceptConnection::Accept)
 	{
 		// Accept the connection with the API. We'll want to do this as quick as possible.
 		EResult AcceptedResult = FSteamSocketsSubsystem::GetSteamSocketsInterface()->AcceptConnection(SocketHandle);
diff --git a/Engine/Source/Runtime/Engine/Private/DataChannel.cpp b/Engine/Source/Runtime/Engine/Private/DataChannel.cpp
--- a/Engine/Source/Runtime/Engine/Private/DataChannel.cpp
+++ b/Engine/Source/Runtime/Engine/Private/DataChannel.cpp
@@ -1793,7 +1793,7 @@ void UControlChannel::ReceivedBunch( FInBunch& Bunch )
 
 			break;
 		}
-		else
+		else if (Connection->Driver->Notify != nullptr)
 		{
 			// Process control message on client/server connection
 			Connection->Driver->Notify->NotifyControlMessage(Connection, MessageType, Bunch);
diff --git a/Engine/Source/Runtime/Engine/Private/NetConnection.cpp b/Engine/Source/Runtime/Engine/Private/NetConnection.cpp
--- a/Engine/Source/Runtime/Engine/Private/NetConnection.cpp
+++ b/Engine/Source/Runtime/Engine/Private/NetConnection.cpp
@@ -3319,7 +3319,7 @@ void UNetConnection::ReceivedPacket( FBitReader& Reader, bool bIsReinjectedPacke
 					Channel = CreateChannelByName( Bunch.ChName, EChannelCreateFlags::None, Bunch.ChIndex );
 
 					// Notify the server of the new channel.
-					if( !Driver->Notify->NotifyAcceptingChannel( Channel ) )
+					if( Driver->Notify == nullptr || !Driver->Notify->NotifyAcceptingChannel( Channel ) )
 					{
 						// Channel refused, so close it, flush it, and delete it.
 						UE_LOG(LogNet, Verbose, TEXT("      NotifyAcceptingChannel Failed! Channel: %s"), *Channel->Describe() );
