Net Address Resolution refactor, to make way for 'Half-broken' NetConnection fix.

diff --git a/Engine/Plugins/Online/OnlineSubsystemSteam/Source/Private/SocketSubsystemSteam.cpp b/Engine/Plugins/Online/OnlineSubsystemSteam/Source/Private/SocketSubsystemSteam.cpp
--- a/Engine/Plugins/Online/OnlineSubsystemSteam/Source/Private/SocketSubsystemSteam.cpp
+++ b/Engine/Plugins/Online/OnlineSubsystemSteam/Source/Private/SocketSubsystemSteam.cpp
@@ -262,13 +262,19 @@ void FSocketSubsystemSteam::RegisterConnection(USteamNetConnection* Connection)
 	FWeakObjectPtr ObjectPtr = Connection;
 	SteamConnections.Add(ObjectPtr);
 
-	if (Connection->GetRemoteAddr().IsValid() && Connection->Socket)
+	if (FSocket* CurSocket = Connection->GetSocket())
 	{
-		FSocketSteam* SteamSocket = (FSocketSteam*)Connection->Socket;
-		TSharedPtr<const FInternetAddrSteam> SteamAddr = StaticCastSharedPtr<const FInternetAddrSteam>(Connection->GetRemoteAddr());
+		TSharedPtr<const FInternetAddr> CurRemoteAddr = Connection->GetRemoteAddr();
+
+		if (CurRemoteAddr.IsValid())
+		{
+			FSocketSteam* SteamSocket = (FSocketSteam*)CurSocket;
+			TSharedPtr<const FInternetAddrSteam> SteamAddr = StaticCastSharedPtr<const FInternetAddrSteam>(CurRemoteAddr);
 
-		UE_LOG_ONLINE(Log, TEXT("Adding user %s from RegisterConnection"), *SteamAddr->ToString(true));
-		P2PTouch(SteamSocket->SteamNetworkingPtr, *SteamAddr->SteamId, SteamAddr->SteamChannel);
+			UE_LOG_ONLINE(Log, TEXT("Adding user %s from RegisterConnection"), *SteamAddr->ToString(true));
+
+			P2PTouch(SteamSocket->SteamNetworkingPtr, *SteamAddr->SteamId, SteamAddr->SteamChannel);
+		}
 	}
 }
 
diff --git a/Engine/Plugins/Online/OnlineSubsystemUtils/Source/OnlineSubsystemUtils/Classes/IpConnection.h b/Engine/Plugins/Online/OnlineSubsystemUtils/Source/OnlineSubsystemUtils/Classes/IpConnection.h
--- a/Engine/Plugins/Online/OnlineSubsystemUtils/Source/OnlineSubsystemUtils/Classes/IpConnection.h
+++ b/Engine/Plugins/Online/OnlineSubsystemUtils/Source/OnlineSubsystemUtils/Classes/IpConnection.h
@@ -6,7 +6,6 @@
 
 #pragma once
 
-#include "CoreMinimal.h"
 #include "UObject/ObjectMacros.h"
 #include "Engine/NetConnection.h"
 #include "Async/TaskGraphInterfaces.h"
@@ -15,9 +14,18 @@
 
 class FInternetAddr;
 class ISocketSubsystem;
+class FSocket;
+class FIpConnectionHelper;
+
+namespace UE::Net::Private
+{
+	class FNetDriverAddressResolution;
+	class FNetConnectionAddressResolution;
+}
+
 
 /** A state system of the address resolution functionality. */
-enum class EAddressResolutionState : uint8
+enum class UE_DEPRECATED(5.1, "EAddressResolutionState has been moved to a private namespace.") EAddressResolutionState : uint8
 {
 	None = 0,
 	Disabled,
@@ -33,13 +41,17 @@ UCLASS(transient, config=Engine)
 class ONLINESUBSYSTEMUTILS_API UIpConnection : public UNetConnection
 {
     GENERATED_UCLASS_BODY()
-	// Variables.
 
+	friend FIpConnectionHelper;
+	friend UE::Net::Private::FNetDriverAddressResolution;
+
+public:
 	/** This is a non-owning pointer to a socket owned elsewhere, IpConnection will not destroy the socket through this pointer. */
-	class FSocket*				Socket;
-	UE_DEPRECATED(4.25, "Address resolution is now handled in the IpNetDriver and no longer done entirely in the IpConnection")
-	class FResolveInfo*			ResolveInfo;
+	UE_DEPRECATED(5.1, "Socket access is now controlled through GetSocket")
+	FSocket*				Socket;
+
 
+public:
 	//~ Begin NetConnection Interface
 	virtual void InitBase(UNetDriver* InDriver, class FSocket* InSocket, const FURL& InURL, EConnectionState InState, int32 InMaxPacket = 0, int32 InPacketOverhead = 0) override;
 	virtual void InitRemoteConnection(UNetDriver* InDriver, class FSocket* InSocket, const FURL& InURL, const class FInternetAddr& InRemoteAddr, EConnectionState InState, int32 InMaxPacket = 0, int32 InPacketOverhead = 0) override;
@@ -59,6 +71,18 @@ class ONLINESUBSYSTEMUTILS_API UIpConnection : public UNetConnection
 	 */
 	void WaitForSendTasks();
 
+	/**
+	 * Gets the cached socket for this connection.
+	 *
+	 * @return	The cached socket for this connection
+	 */
+	FSocket* GetSocket() const
+	{
+		PRAGMA_DISABLE_DEPRECATION_WARNINGS
+		return Socket;
+		PRAGMA_ENABLE_DEPRECATION_WARNINGS
+	}
+
 private:
 	/**
 	 * Struct to hold the result of a socket SendTo call. If net.IpConnectionUseSendTasks is true,
@@ -100,69 +124,24 @@ private:
 
 private:
 
+	/** Cleanup for the deprecated 'Socket' value */
+	void CleanupDeprecatedSocket();
+
 	/** Handles any SendTo errors on the game thread. */
 	void HandleSocketSendResult(const FSocketSendResult& Result, ISocketSubsystem* SocketSubsystem);
 
 	/** Notifies us that we've encountered an error while receiving a packet. */
 	void HandleSocketRecvError(class UNetDriver* NetDriver, const FString& ErrorString);
 
-	/** An array of sockets tied to every binding address. */
-	TArray<TSharedPtr<FSocket>> BindSockets;
-
-	/** Holds a refcount to the actual socket to be used from BindSockets. */
-	TSharedPtr<FSocket> ResolutionSocket;
-
-	/** An array containing the address results GAI returns for the current host value. Given to us from the netdriver. */
-	TArray<TSharedRef<FInternetAddr>> ResolverResults;
+	/** NetDriver level early address resolution (may pass work on to NetConnection level address resolution) */
+	TPimplPtr<UE::Net::Private::FNetConnectionAddressResolution> Resolver;
 
-	/** The index into the ResolverResults that we're currently attempting */
-	int32 CurrentAddressIndex;
-
-	/** 
-	 *  The connection's current status of where it is in the resolution state machine.
-	 *  If a platform should not use resolution, call DisableAddressResolution() in your constructor
-	 */
-	EAddressResolutionState ResolutionState;
-
-	/**
-	 * Cleans up the socket information in use with resolution. This can get called numerous times.
-	 */
-	void CleanupResolutionSockets();
-
-	/**
-	 * Determines if we can continue processing resolution results or not based on flags and
-	 * current flow.
-	 *
-	 * @return if resolution is allowed to continue processing.
-	 */
-	bool CanContinueResolution() const {
-		return CurrentAddressIndex < ResolverResults.Num() && IsAddressResolutionEnabled() &&
-			ResolutionState != EAddressResolutionState::Error && ResolutionState != EAddressResolutionState::Done;
-	}
-
-	/**
-	 * Checks to see if this netconnection class can use address resolution
-	 *
-	 * @return if address resolution is allowed to continue processing.
-	 */
-	bool IsAddressResolutionEnabled() const {
-		return ResolutionState != EAddressResolutionState::Disabled;
-	}
-
-	/**
-	 * Checks to see if this netconnection class has encountered an error during process resolution
-	 *
-	 * @return If an error has occurred
-	 */
-	bool HasAddressResolutionFailed() const {
-		return ResolutionState == EAddressResolutionState::Error;
-	}
 
 protected:
 	/**
 	 * Disables address resolution by pushing the disabled flag into the status field.
 	 */
-	void DisableAddressResolution() { ResolutionState = EAddressResolutionState::Disabled; }
+	void DisableAddressResolution();
 
 	/**
 	 * Handles a NetConnection timeout. Overridden in order to handle parsing multiple GAI results during resolution.
@@ -170,6 +149,4 @@ protected:
 	 * @param ErrorStr A string containing the current error message for either usage or writing into.
 	 */
 	virtual void HandleConnectionTimeout(const FString& ErrorStr) override;
-
-	friend class FIpConnectionHelper;
 };
diff --git a/Engine/Plugins/Online/OnlineSubsystemUtils/Source/OnlineSubsystemUtils/Classes/IpNetDriver.h b/Engine/Plugins/Online/OnlineSubsystemUtils/Source/OnlineSubsystemUtils/Classes/IpNetDriver.h
--- a/Engine/Plugins/Online/OnlineSubsystemUtils/Source/OnlineSubsystemUtils/Classes/IpNetDriver.h
+++ b/Engine/Plugins/Online/OnlineSubsystemUtils/Source/OnlineSubsystemUtils/Classes/IpNetDriver.h
@@ -6,7 +6,6 @@
 
 #pragma once
 
-#include "CoreMinimal.h"
 #include "UObject/ObjectMacros.h"
 #include "Engine/NetDriver.h"
 #include "HAL/Runnable.h"
@@ -24,6 +23,11 @@ class FNetworkNotify;
 class FSocket;
 struct FRecvMulti;
 
+namespace UE::Net::Private
+{
+	class FNetDriverAddressResolution;
+}
+
 
 // CVars
 #if !UE_BUILD_SHIPPING
@@ -51,10 +55,6 @@ public:
 	UPROPERTY(Config)
 	uint32 MaxPortCountToTry;
 
-	/** Underlying socket communication */
-	UE_DEPRECATED(4.25, "Socket access is now controlled through the getter/setter combo: GetSocket and SetSocketAndLocalAddress")
-	FSocket* Socket;
-
 	/** If pausing socket receives, the time at which this should end */
 	float PauseReceiveEnd;
 
@@ -77,30 +77,6 @@ public:
 	}
 	//~ End UNetDriver Interface
 
-	/**
-	 * Processes packets not associated with a NetConnection, performing any handshaking and NetConnection creation or remapping, as necessary.
-	 *
-	 * @param Address			The address the packet came from
-	 * @param Data				The packet data (may be modified)
-	 * @param CountBytesRef		The packet size (may be modified)
-	 * @return					If a new NetConnection is created, returns the net connection
-	 */
-	UE_DEPRECATED(4.24, "ProcessConnectionlessPacket has a new API, and will become private soon.")
-	UNetConnection* ProcessConnectionlessPacket(const TSharedRef<FInternetAddr>& Address, uint8* Data, int32& CountBytesRef)
-	{
-		FReceivedPacketView PacketView;
-
-		PacketView.DataView = {Data, CountBytesRef, ECountUnits::Bytes};
-		PacketView.Address = Address;
-		PacketView.Error = SE_NO_ERROR;
-
-		UNetConnection* ReturnVal = ProcessConnectionlessPacket(PacketView, { Data, MAX_PACKET_SIZE } );
-
-		CountBytesRef = PacketView.DataView.NumBytes();
-
-		return ReturnVal;
-	}
-
 private:
 	/**
 	 * Process packets not associated with a NetConnection, performing handshaking and NetConnection creation or remapping as necessary.
@@ -117,14 +93,6 @@ private:
 public:
 	//~ Begin UIpNetDriver Interface.
 
-	/**
-	 * Creates a socket to be used for network communications. Uses the LocalAddr (if set) to determine protocol flags
-	 *
-	 * @return an FSocket if creation succeeded, nullptr if creation failed.
-	 */
-	UE_DEPRECATED(4.25, "Socket creation is now restricted to subclasses and done through CreateSocketForProtocol or CreateAndBindSocket")
-	virtual FSocket* CreateSocket();
-	
 	/**
 	 * Returns the current FSocket to be used with this NetDriver. This is useful in the cases of resolution as it will always point to the Socket that's currently being
 	 * used with the current resolution attempt.
@@ -337,12 +305,6 @@ private:
 	/** The preallocated state/buffers, for efficiently executing RecvMulti */
 	TUniquePtr<FRecvMulti> RecvMultiState;
 
-	/** 
-	 * An array sockets created for every binding address a machine has in use for performing address resolution. 
-	 * This array empties after connections have been spun up.
-	 */
-	TArray<TSharedPtr<FSocket>> BoundSockets;
-
 	/** Underlying socket communication */
 	TSharedPtr<FSocket> SocketPrivate;
 
@@ -371,4 +333,7 @@ private:
 
 	/** Countdown timer for the current IP aggregation tracking period */
 	double NextAggregateIPLogCountdown = 0.0;
+
+	/** NetConnection specific address resolution */
+	TPimplPtr<UE::Net::Private::FNetDriverAddressResolution> Resolver;
 };
diff --git a/Engine/Plugins/Online/OnlineSubsystemUtils/Source/OnlineSubsystemUtils/Private/IpConnection.cpp b/Engine/Plugins/Online/OnlineSubsystemUtils/Source/OnlineSubsystemUtils/Private/IpConnection.cpp
--- a/Engine/Plugins/Online/OnlineSubsystemUtils/Source/OnlineSubsystemUtils/Private/IpConnection.cpp
+++ b/Engine/Plugins/Online/OnlineSubsystemUtils/Source/OnlineSubsystemUtils/Private/IpConnection.cpp
@@ -20,6 +20,8 @@ Notes:
 
 #include "Net/Core/Misc/PacketAudit.h"
 #include "Misc/ScopeExit.h"
+#include "NetAddressResolution.h"
+
 
 /*-----------------------------------------------------------------------------
 	Declarations.
@@ -32,7 +34,6 @@ Notes:
 DECLARE_CYCLE_STAT(TEXT("IpConnection InitRemoteConnection"), Stat_IpConnectionInitRemoteConnection, STATGROUP_Net);
 DECLARE_CYCLE_STAT(TEXT("IpConnection Socket SendTo"), STAT_IpConnection_SendToSocket, STATGROUP_Net);
 DECLARE_CYCLE_STAT(TEXT("IpConnection WaitForSendTasks"), STAT_IpConnection_WaitForSendTasks, STATGROUP_Net);
-DECLARE_CYCLE_STAT(TEXT("IpConnection Address Synthesis"), STAT_IpConnection_AddressSynthesis, STATGROUP_Net);
 
 TAutoConsoleVariable<int32> CVarNetIpConnectionUseSendTasks(
 	TEXT("net.IpConnectionUseSendTasks"),
@@ -45,16 +46,22 @@ TAutoConsoleVariable<int32> CVarNetIpConnectionDisableResolution(
 	TEXT("If enabled, any future ip connections will not use resolution methods."),
 	ECVF_Default | ECVF_Cheat);
 
-PRAGMA_DISABLE_DEPRECATION_WARNINGS
+
+/**
+ * UIpConnection
+ */
+
 UIpConnection::UIpConnection(const FObjectInitializer& ObjectInitializer) :
 	Super(ObjectInitializer),
+	PRAGMA_DISABLE_DEPRECATION_WARNINGS
 	Socket(nullptr),
-	ResolveInfo(nullptr),
+	PRAGMA_ENABLE_DEPRECATION_WARNINGS
 	SocketErrorDisconnectDelay(5.f),
 	SocketError_SendDelayStartTime(0.f),
 	SocketError_RecvDelayStartTime(0.f),
-	CurrentAddressIndex(0),
-	ResolutionState(EAddressResolutionState::None)
+	PRAGMA_DISABLE_DEPRECATION_WARNINGS
+	Resolver(MakePimpl<UE::Net::Private::FNetConnectionAddressResolution>(Socket))
+	PRAGMA_ENABLE_DEPRECATION_WARNINGS
 {
 	// Auto add address resolution disable flags if the cvar is set.
 	if (!!CVarNetIpConnectionDisableResolution.GetValueOnAnyThread())
@@ -62,7 +69,6 @@ UIpConnection::UIpConnection(const FObjectInitializer& ObjectInitializer) :
 		DisableAddressResolution();
 	}
 }
-PRAGMA_ENABLE_DEPRECATION_WARNINGS
 
 void UIpConnection::InitBase(UNetDriver* InDriver, class FSocket* InSocket, const FURL& InURL, EConnectionState InState, int32 InMaxPacket, int32 InPacketOverhead)
 {
@@ -72,9 +78,11 @@ void UIpConnection::InitBase(UNetDriver* InDriver, class FSocket* InSocket, cons
 		(InMaxPacket == 0 || InMaxPacket > MAX_PACKET_SIZE) ? MAX_PACKET_SIZE : InMaxPacket,
 		InPacketOverhead == 0 ? UDP_HEADER_SIZE : InPacketOverhead);
 
-	if (!IsAddressResolutionEnabled())
+	if (!Resolver->IsAddressResolutionEnabled())
 	{
+		PRAGMA_DISABLE_DEPRECATION_WARNINGS
 		Socket = InSocket;
+		PRAGMA_ENABLE_DEPRECATION_WARNINGS
 	}
 
 	if (CVarNetEnableCongestionControl.GetValueOnAnyThread() > 0)
@@ -90,51 +98,18 @@ void UIpConnection::InitLocalConnection(UNetDriver* InDriver, class FSocket* InS
 		(InMaxPacket == 0 || InMaxPacket > MAX_PACKET_SIZE) ? MAX_PACKET_SIZE : InMaxPacket,
 		InPacketOverhead == 0 ? UDP_HEADER_SIZE : InPacketOverhead);
 
-	// If resolution is disabled, fall back to address synthesis
-	if (!IsAddressResolutionEnabled())
-	{
-		// Figure out IP address from the host URL
-		bool bIsValid = false;
-		// Get numerical address directly.
-		RemoteAddr = InDriver->GetSocketSubsystem()->CreateInternetAddr();
-		RemoteAddr->SetIp(*InURL.Host, bIsValid);
-
-		// If the protocols do not match, attempt to synthesize the address so they do.
-		if ((bIsValid && InSocket->GetProtocol() != RemoteAddr->GetProtocolType()) || !bIsValid)
-		{
-			SCOPE_CYCLE_COUNTER(STAT_IpConnection_AddressSynthesis);
-
-			// We want to use GAI to create the address with the correct protocol.
-			const FAddressInfoResult MapRequest = InDriver->GetSocketSubsystem()->GetAddressInfo(*InURL.Host, nullptr,
-				EAddressInfoFlags::AllResultsWithMapping | EAddressInfoFlags::OnlyUsableAddresses, InSocket->GetProtocol());
-
-			// Set the remote addr provided we have information.
-			if (MapRequest.ReturnCode == SE_NO_ERROR && MapRequest.Results.Num() > 0)
-			{
-				RemoteAddr = MapRequest.Results[0].Address->Clone();
-				bIsValid = true;
-			}
-			else
-			{
-				UE_LOG(LogNet, Warning, TEXT("IpConnection::InitConnection: Address protocols do not match and cannot be synthesized to a similar address, this will likely lead to issues!"));
-			}
-		}
 
-		RemoteAddr->SetPort(InURL.Port);
+	bool bResolverInit = Resolver->InitLocalConnection(InDriver->GetSocketSubsystem(), InSocket, InURL);
 
-		// If synthesis failed then we shouldn't continue.
-		if (bIsValid == false)
-		{
-			Close();
-			UE_LOG(LogNet, Verbose, TEXT("IpConnection::InitConnection: Unable to resolve %s"), *InURL.Host);
-			return;
-		}
-	}
-	else
+	if (!bResolverInit)
 	{
-		ResolutionState = EAddressResolutionState::WaitingForResolves;
+		Close();
+
+		return;
 	}
 
+	RemoteAddr = Resolver->GetRemoteAddr();
+
 	// Initialize our send bunch
 	InitSendBuffer();
 }
@@ -143,8 +118,8 @@ void UIpConnection::InitRemoteConnection(UNetDriver* InDriver, class FSocket* In
 {
 	SCOPE_CYCLE_COUNTER(Stat_IpConnectionInitRemoteConnection);
 
-	// Listeners don't need to perform address resolution, so flag it as disabled.
-	ResolutionState = EAddressResolutionState::Disabled;
+	// Listeners don't need to perform address resolution
+	Resolver->DisableAddressResolution();
 
 	InitBase(InDriver, InSocket, InURL, InState, 
 		// Use the default packet size/overhead unless overridden by a child class
@@ -164,6 +139,8 @@ void UIpConnection::InitRemoteConnection(UNetDriver* InDriver, class FSocket* In
 
 void UIpConnection::Tick(float DeltaSeconds)
 {
+	using namespace UE::Net::Private;
+
 	if (CVarNetIpConnectionUseSendTasks.GetValueOnGameThread() != 0)
 	{
 		ISocketSubsystem* const SocketSubsystem = Driver->GetSocketSubsystem();
@@ -184,36 +161,16 @@ void UIpConnection::Tick(float DeltaSeconds)
 		}
 	}
 
-	if (ResolutionState == EAddressResolutionState::TryNextAddress)
-	{
-		if (CurrentAddressIndex >= ResolverResults.Num())
-		{
-			UE_LOG(LogNet, Warning, TEXT("Exhausted the number of resolver results, closing the connection now."));
-			ResolutionState = EAddressResolutionState::Error;
-			return;
-		}
 
-		RemoteAddr = ResolverResults[CurrentAddressIndex];
-		Socket = nullptr;
+	ECheckAddressResolutionResult CheckResult = Resolver->CheckAddressResolution();
 
-		for (const TSharedPtr<FSocket>& BindSocket : BindSockets)
-		{
-			if (BindSocket->GetProtocol() == RemoteAddr->GetProtocolType())
-			{
-				ResolutionSocket = BindSocket;
-				Socket = ResolutionSocket.Get();
-				break;
-			}
-		}
+	if (CheckResult == ECheckAddressResolutionResult::TryFirstAddress || CheckResult == ECheckAddressResolutionResult::TryNextAddress)
+	{
+		PRAGMA_DISABLE_DEPRECATION_WARNINGS
+		Socket = Resolver->GetResolutionSocket().Get();
+		PRAGMA_ENABLE_DEPRECATION_WARNINGS
 
-		if (Socket == nullptr)
-		{
-			UE_LOG(LogNet, Error, TEXT("Unable to find a binding socket for the resolve address result %s"), *RemoteAddr->ToString(true));
-			ResolutionState = EAddressResolutionState::Error;
-			return;
-		}
-		
-		ResolutionState = EAddressResolutionState::Connecting;
+		RemoteAddr = Resolver->GetRemoteAddr();
 
 		// Reset any timers
 		LastReceiveTime = Driver->GetElapsedTime();
@@ -225,7 +182,7 @@ void UIpConnection::Tick(float DeltaSeconds)
 
 		// Resend initial packets again (only need to do this if this connection does not have packethandlers)
 		// Otherwise most packethandlers will do retry (the stateless handshake being one example)
-		if (CurrentAddressIndex != 0 && !Handler.IsValid())
+		if (CheckResult == ECheckAddressResolutionResult::TryNextAddress && !Handler.IsValid())
 		{
 			FWorldContext* WorldContext = GEngine->GetWorldContextFromPendingNetGameNetDriver(Driver);
 			if (WorldContext != nullptr && WorldContext->PendingNetGame != nullptr)
@@ -233,27 +190,23 @@ void UIpConnection::Tick(float DeltaSeconds)
 				WorldContext->PendingNetGame->SendInitialJoin();
 			}
 		}
-
-		++CurrentAddressIndex;
 	}
-	else if (ResolutionState == EAddressResolutionState::Connected)
+	else if (CheckResult == ECheckAddressResolutionResult::Connected)
 	{
-		CleanupResolutionSockets();
 		UIpNetDriver* IpDriver = Cast<UIpNetDriver>(Driver);
 
 		// Set the right object now that we have a connection
-		IpDriver->SetSocketAndLocalAddress(ResolutionSocket);
-
-		ResolutionState = EAddressResolutionState::Done;
+		IpDriver->SetSocketAndLocalAddress(Resolver->GetResolutionSocket());
 	}
-	else if (ResolutionState == EAddressResolutionState::Error)
+	else if (CheckResult == ECheckAddressResolutionResult::Error)
 	{
-		UE_LOG(LogNet, Warning, TEXT("Encountered an error, cleaning up this connection now"));
-		ResolutionState = EAddressResolutionState::Done;
-
 		// Host name resolution just now failed.
 		SetConnectionState(USOCK_Closed);
-		Close();
+		Close(ENetCloseResult::AddressResolutionFailed);
+	}
+	else if (CheckResult == ECheckAddressResolutionResult::FindSocketError)
+	{
+		CleanupDeprecatedSocket();
 	}
 
 	Super::Tick(DeltaSeconds);
@@ -261,30 +214,40 @@ void UIpConnection::Tick(float DeltaSeconds)
 
 void UIpConnection::CleanUp()
 {
-	// Force ourselves into a finishing state such that we clean up any excess network objects.
-	if (IsAddressResolutionEnabled())
-	{
-		ResolutionState = EAddressResolutionState::Done;
-	}
-
-	// Clean up these sockets now, as we'll lose the NetDriver pointer later on
-	CleanupResolutionSockets();
+	using namespace UE::Net::Private;
 
 	Super::CleanUp();
 
 	WaitForSendTasks();
+
+	// Sockets must be cleaned up after send tasks complete, as they can still be in use
+	CleanupDeprecatedSocket();
+	Resolver->CleanupResolutionSockets();
+}
+
+void UIpConnection::CleanupDeprecatedSocket()
+{
+	PRAGMA_DISABLE_DEPRECATION_WARNINGS
+	if (Socket != nullptr && Resolver->IsAddressResolutionEnabled())
+	{
+		Socket = nullptr;
+	}
+	PRAGMA_ENABLE_DEPRECATION_WARNINGS
 }
 
 void UIpConnection::HandleConnectionTimeout(const FString& ErrorStr)
 {
-	if (CanContinueResolution())
+	using namespace UE::Net::Private;
+
+	const EEAddressResolutionHandleResult HandleResult = Resolver->NotifyTimeout();
+
+	if (HandleResult == EEAddressResolutionHandleResult::CallerShouldHandle)
 	{
-		ResolutionState = EAddressResolutionState::TryNextAddress;
-		bPendingDestroy = false;
+		Super::HandleConnectionTimeout(ErrorStr);
 	}
 	else
 	{
-		Super::HandleConnectionTimeout(ErrorStr);
+		bPendingDestroy = false;
 	}
 }
 
@@ -333,7 +296,7 @@ void UIpConnection::LowLevelSend(void* Data, int32 CountBits, FOutPacketTraits&
 	};
 #endif
 
-	if (IsAddressResolutionEnabled() && (ResolutionState == EAddressResolutionState::WaitingForResolves || ResolutionState == EAddressResolutionState::TryNextAddress))
+	if (Resolver->ShouldBlockSend())
 	{
 		UE_LOG(LogNet, Verbose, TEXT("Skipping send task as we are waiting on the next resolution step"));
 		return;
@@ -378,6 +341,7 @@ void UIpConnection::LowLevelSend(void* Data, int32 CountBits, FOutPacketTraits&
 		if (CountBytes > 0)
 		{
 			const bool bNotifyOnSuccess = (SocketErrorDisconnectDelay > 0.f) && (SocketError_SendDelayStartTime != 0.f);
+			FSocket* CurSocket = GetSocket();
 
 			if (CVarNetIpConnectionUseSendTasks.GetValueOnAnyThread() != 0)
 			{
@@ -393,14 +357,16 @@ void UIpConnection::LowLevelSend(void* Data, int32 CountBits, FOutPacketTraits&
 				
 				LastSendTask = FFunctionGraphTask::CreateAndDispatchWhenReady([this, Packet = TArray<uint8>(DataToSend, CountBytes), SocketSubsystem, bNotifyOnSuccess]
 				{
-					if (Socket != nullptr)
+					FSocket* CurSocket = GetSocket();
+
+					if (CurSocket != nullptr)
 					{
 						bool bWasSendSuccessful = false;
 						UIpConnection::FSocketSendResult Result;
 
 						{
 							SCOPE_CYCLE_COUNTER(STAT_IpConnection_SendToSocket);
-							bWasSendSuccessful = Socket->SendTo(Packet.GetData(), Packet.Num(), Result.BytesSent, *RemoteAddr);
+							bWasSendSuccessful = CurSocket->SendTo(Packet.GetData(), Packet.Num(), Result.BytesSent, *RemoteAddr);
 						}
 
 						if (!bWasSendSuccessful && SocketSubsystem)
@@ -420,21 +386,21 @@ void UIpConnection::LowLevelSend(void* Data, int32 CountBits, FOutPacketTraits&
 				// Always flush this profiler data now. Technically this could be incorrect if the send in the task fails,
 				// but this keeps the bookkeeping simpler for now.
 				NETWORK_PROFILER(GNetworkProfiler.FlushOutgoingBunches(this));
-				NETWORK_PROFILER(GNetworkProfiler.TrackSocketSendTo(Socket->GetDescription(), DataToSend, CountBytes, NumPacketIdBits, NumBunchBits, NumAckBits, NumPaddingBits, this));
+				NETWORK_PROFILER(GNetworkProfiler.TrackSocketSendTo(CurSocket->GetDescription(), DataToSend, CountBytes, NumPacketIdBits, NumBunchBits, NumAckBits, NumPaddingBits, this));
 			}
 			else
 			{
 				bool bWasSendSuccessful = false;
 				{
 					SCOPE_CYCLE_COUNTER(STAT_IpConnection_SendToSocket);
-					bWasSendSuccessful = Socket->SendTo(DataToSend, CountBytes, SendResult.BytesSent, *RemoteAddr);
+					bWasSendSuccessful = CurSocket->SendTo(DataToSend, CountBytes, SendResult.BytesSent, *RemoteAddr);
 				}
 
 				if (bWasSendSuccessful)
 				{
 					UNCLOCK_CYCLES(Driver->SendCycles);
 					NETWORK_PROFILER(GNetworkProfiler.FlushOutgoingBunches(this));
-					NETWORK_PROFILER(GNetworkProfiler.TrackSocketSendTo(Socket->GetDescription(), DataToSend, SendResult.BytesSent, NumPacketIdBits, NumBunchBits, NumAckBits, NumPaddingBits, this));
+					NETWORK_PROFILER(GNetworkProfiler.TrackSocketSendTo(CurSocket->GetDescription(), DataToSend, SendResult.BytesSent, NumPacketIdBits, NumBunchBits, NumAckBits, NumPaddingBits, this));
 
 					if (bNotifyOnSuccess)
 					{
@@ -455,16 +421,17 @@ void UIpConnection::LowLevelSend(void* Data, int32 CountBits, FOutPacketTraits&
 
 void UIpConnection::ReceivedRawPacket(void* Data, int32 Count)
 {
-	UE_CLOG(SocketError_RecvDelayStartTime > 0.0, LogNet, Log, TEXT("UIpConnection::ReceivedRawPacket: Recoverd from socket errors. %s Connection"), *Describe());
+	UE_CLOG(SocketError_RecvDelayStartTime > 0.0, LogNet, Log,
+			TEXT("UIpConnection::ReceivedRawPacket: Recovered from socket errors. %s Connection"), ToCStr(Describe()));
 	
 	// We received data successfully, reset our error counters.
 	SocketError_RecvDelayStartTime = 0.0;
 
 	// Set that we've gotten packet from the server, this begins destruction of the other elements.
-	if (IsAddressResolutionEnabled() && ResolutionState != EAddressResolutionState::Done)
+	if (Resolver->IsAddressResolutionEnabled() && !Resolver->IsAddressResolutionComplete())
 	{
 		// We only want to write this once, because we don't want to waste cycles trying to clean up nothing.
-		ResolutionState = EAddressResolutionState::Connected;
+		Resolver->NotifyAddressResolutionConnected();
 	}
 
 	Super::ReceivedRawPacket(Data, Count);
@@ -481,13 +448,15 @@ float UIpConnection::GetTimeoutValue()
 	}
 #endif
 
-	if (ResolutionState == EAddressResolutionState::Connecting)
+	if (Resolver->IsAddressResolutionConnecting())
 	{
 		UIpNetDriver* IpDriver = Cast<UIpNetDriver>(Driver);
+
 		if (IpDriver == nullptr || IpDriver->GetResolutionTimeoutValue() == 0.0f)
 		{
 			return Driver->InitialConnectTimeout;
 		}
+
 		return IpDriver->GetResolutionTimeoutValue();
 	}
 
@@ -496,6 +465,8 @@ float UIpConnection::GetTimeoutValue()
 
 void UIpConnection::HandleSocketSendResult(const FSocketSendResult& Result, ISocketSubsystem* SocketSubsystem)
 {
+	using namespace UE::Net::Private;
+
 	if (Result.Error == SE_NO_ERROR)
 	{
 		UE_CLOG(SocketError_SendDelayStartTime > 0.f, LogNet, Log, TEXT("UIpConnection::HandleSocketSendResult: Recovered from socket errors. %s Connection"), *Describe());
@@ -529,36 +500,35 @@ void UIpConnection::HandleSocketSendResult(const FSocketSendResult& Result, ISoc
 		}
 
 		// Broadcast the error only on the first occurrence
-		if( GetPendingCloseDueToSocketSendFailure() == false )
+		if (!GetPendingCloseDueToSocketSendFailure())
 		{
-			if (CanContinueResolution())
-			{
-				SetPendingCloseDueToSocketSendFailure();
-				SocketError_SendDelayStartTime = 0.f;
-				ResolutionState = EAddressResolutionState::TryNextAddress;
-			}
-			else
-			{
-				ResolutionState = EAddressResolutionState::Error;
+			const EEAddressResolutionHandleResult HandleResult = Resolver->NotifySendError();
 
-				// Request the connection to be disconnected during next tick() since we got a critical socket failure, the actual disconnect is postponed 		
-				// to avoid issues with the call Close() causing issues with reentrant code paths in DataChannel::SendBunch() and FlushNet()
-				SetPendingCloseDueToSocketSendFailure();
+			// Request the connection to be disconnected during next tick() since we got a critical socket failure, the actual disconnect is postponed 		
+			// to avoid issues with the call Close() causing issues with reentrant code paths in DataChannel::SendBunch() and FlushNet()
+			SetPendingCloseDueToSocketSendFailure();
 
-				FString ErrorString = FString::Printf(TEXT("UIpNetConnection::HandleSocketSendResult: Socket->SendTo failed with error %i (%s). %s Connection will be closed during next Tick()!"),
-					static_cast<int32>(Result.Error),
-					SocketSubsystem->GetSocketError(Result.Error),
-					*Describe());
+			if (HandleResult == EEAddressResolutionHandleResult::CallerShouldHandle)
+			{
+				FString ErrorString = FString::Printf(TEXT("UIpNetConnection::HandleSocketSendResult: Socket->SendTo failed with error %i (%s). ")
+														TEXT("%s Connection will be closed during next Tick()!"),
+														static_cast<int32>(Result.Error), ToCStr(SocketSubsystem->GetSocketError(Result.Error)),
+														ToCStr(Describe()));
 
 				GEngine->BroadcastNetworkFailure(Driver->GetWorld(), Driver, ENetworkFailure::ConnectionLost, ErrorString);
 			}
-
+			else
+			{
+				SocketError_SendDelayStartTime = 0.f;
+			}
 		}
 	}
 }
 
 void UIpConnection::HandleSocketRecvError(class UNetDriver* NetDriver, const FString& ErrorString)
 {
+	using namespace UE::Net::Private;
+
 	check(NetDriver);
 
 	if (SocketErrorDisconnectDelay > 0.f)
@@ -580,49 +550,25 @@ void UIpConnection::HandleSocketRecvError(class UNetDriver* NetDriver, const FSt
 		}
 	}
 
-	if (CanContinueResolution())
-	{
-		SocketErrorDisconnectDelay = 0.0f;
-		ResolutionState = EAddressResolutionState::TryNextAddress;
-	}
-	else
-	{
-		ResolutionState = EAddressResolutionState::Error;
 
+	const EEAddressResolutionHandleResult HandleResult = Resolver->NotifyReceiveError();
+
+	if (HandleResult == EEAddressResolutionHandleResult::CallerShouldHandle)
+	{
 		// For now, this is only called on clients when the ServerConnection fails.
 		// Because of that, on failure we'll shut down the NetDriver.
 		GEngine->BroadcastNetworkFailure(NetDriver->GetWorld(), NetDriver, ENetworkFailure::ConnectionLost, ErrorString);
 		NetDriver->Shutdown();
 	}
-
-}
-
-void UIpConnection::CleanupResolutionSockets()
-{
-	if (Driver == nullptr || Driver->GetSocketSubsystem() == nullptr || !IsAddressResolutionEnabled())
-	{
-		return;
-	}
-
-	const bool bCleanAll = (ResolutionState != EAddressResolutionState::Connected);
-	if (Socket == nullptr && !bCleanAll)
+	else
 	{
-		UE_LOG(LogNet, Warning, TEXT("Cannot clean up resolution sockets as our current socket pointer is null!"));
-		return;
+		SocketErrorDisconnectDelay = 0.0f;
 	}
+}
 
-	if (bCleanAll)
-	{
-		if(ResolutionSocket.Get() == Socket)
-		{
-			Socket = nullptr;
-		}
-		
-		ResolutionSocket.Reset();
-	}
-	
-	BindSockets.Reset();
-	ResolverResults.Empty();
+void UIpConnection::DisableAddressResolution()
+{
+	Resolver->DisableAddressResolution();
 }
 
 FString UIpConnection::LowLevelGetRemoteAddress(bool bAppendPort)
@@ -634,9 +580,9 @@ FString UIpConnection::LowLevelDescribe()
 {
 	TSharedRef<FInternetAddr> LocalAddr = Driver->GetSocketSubsystem()->CreateInternetAddr();
 
-	if (Socket != nullptr)
+	if (FSocket* CurSocket = GetSocket())
 	{
-		Socket->GetAddress(*LocalAddr);
+		CurSocket->GetAddress(*LocalAddr);
 	}
 
 	return FString::Printf
diff --git a/Engine/Plugins/Online/OnlineSubsystemUtils/Source/OnlineSubsystemUtils/Private/IpNetDriver.cpp b/Engine/Plugins/Online/OnlineSubsystemUtils/Source/OnlineSubsystemUtils/Private/IpNetDriver.cpp
--- a/Engine/Plugins/Online/OnlineSubsystemUtils/Source/OnlineSubsystemUtils/Private/IpNetDriver.cpp
+++ b/Engine/Plugins/Online/OnlineSubsystemUtils/Source/OnlineSubsystemUtils/Private/IpNetDriver.cpp
@@ -28,6 +28,7 @@ Notes:
 #include "Serialization/ArchiveCountMem.h"
 #include "Algo/IndexOf.h"
 #include <limits>
+#include "NetAddressResolution.h"
 
 /** For backwards compatibility with the engine stateless connect code */
 #ifndef STATELESSCONNECT_HAS_RANDOM_SEQUENCE
@@ -118,14 +119,6 @@ TAutoConsoleVariable<int32> CVarNetDebugDualIPs(
 		TEXT("(only supports a single client on the server)."));
 
 TSharedPtr<FInternetAddr> GCurrentDuplicateIP;
-
-TAutoConsoleVariable<FString> CVarNetDebugAddResolverAddress(
-	TEXT("net.DebugAppendResolverAddress"),
-	TEXT(""),
-	TEXT("If this is set, all IP address resolution methods will add the value of this CVAR to the list of results.")
-		TEXT("This allows for testing resolution functionality across all multiple addresses with the end goal of having a successful result")
-		TEXT("(being the value of this CVAR)"),
-	ECVF_Default | ECVF_Cheat);
 #endif
 
 namespace IPNetDriverInternal
@@ -713,60 +706,20 @@ private:
 class FIpConnectionHelper
 {
 private:
-	friend class UIpNetDriver;
+	friend UIpNetDriver;
+
 	static void HandleSocketRecvError(UIpNetDriver* Driver, UIpConnection* Connection, const FString& ErrorString)
 	{
 		Connection->HandleSocketRecvError(Driver, ErrorString);
 	}
 
-	static void PushSocketsToConnection(UIpConnection* Connection, TArray<TSharedPtr<FSocket>>& Sockets)
-	{
-		UE_LOG(LogNet, Verbose, TEXT("Pushed %d sockets to net connection %s"), Sockets.Num(), *Connection->GetName());
-		Connection->BindSockets = Sockets;
-	}
-
-	static void PushResolverResultsToConnection(UIpConnection* Connection, TArray<TSharedRef<FInternetAddr>>& ResolverResults)
-	{
-		UE_LOG(LogNet, Verbose, TEXT("Pushed %d resolver results to net connection %s"), ResolverResults.Num(), *Connection->GetName());
-		Connection->ResolverResults = ResolverResults;
-		Connection->ResolutionState = EAddressResolutionState::TryNextAddress;
-	}
-
-	static void CleanUpConnectionSockets(UIpConnection* Connection)
-	{
-		if (Connection != nullptr)
-		{
-			Connection->CleanupResolutionSockets();
-		}
-	}
-
-	static void HandleResolverError(UIpConnection* Connection)
+	static void CleanupDeprecatedSocket(UIpConnection* Connection)
 	{
-		Connection->ResolutionState = EAddressResolutionState::Error;
-		Connection->Close();
-	}
-
-	static bool IsAddressResolutionEnabledForConnection(const UIpConnection* Connection)
-	{
-		if (Connection != nullptr)
-		{
-			return Connection->IsAddressResolutionEnabled();
-		}
-
-		return false;
-	}
-
-	static bool HasAddressResolutionFailedForConnection(const UIpConnection* Connection)
-	{
-		if (Connection != nullptr)
-		{
-			return Connection->HasAddressResolutionFailed();
-		}
-
-		return false;
+		Connection->CleanupDeprecatedSocket();
 	}
 };
 
+
 /**
  * UIpNetDriver
  */
@@ -779,6 +732,7 @@ UIpNetDriver::UIpNetDriver(const FObjectInitializer& ObjectInitializer)
 	, ClientDesiredSocketReceiveBufferBytes(0x8000)
 	, ClientDesiredSocketSendBufferBytes(0x8000)
 	, RecvMultiState(nullptr)
+	, Resolver(MakePimpl<UE::Net::Private::FNetDriverAddressResolution>())
 {
 	using namespace IPNetDriverInternal;
 
@@ -798,12 +752,6 @@ ISocketSubsystem* UIpNetDriver::GetSocketSubsystem()
 	return ISocketSubsystem::Get();
 }
 
-FSocket * UIpNetDriver::CreateSocket()
-{
-	// This is a deprecated function with unsafe socket lifetime management. The Release call is intentional and for backwards compatiblity only.
-	return CreateSocketForProtocol((LocalAddr.IsValid() ? LocalAddr->GetProtocolType() : NAME_None)).Release();
-}
-
 FUniqueSocket UIpNetDriver::CreateSocketForProtocol(const FName& ProtocolType)
 {
 	// Create UDP socket and enable broadcasting.
@@ -903,10 +851,6 @@ void UIpNetDriver::SetSocketAndLocalAddress(const TSharedPtr<FSocket>& SharedSoc
 {
 	SocketPrivate = SharedSocket;
 
-	PRAGMA_DISABLE_DEPRECATION_WARNINGS
-	Socket = SocketPrivate.Get();
-	PRAGMA_ENABLE_DEPRECATION_WARNINGS
-
 	if (SocketPrivate.IsValid())
 	{
 		// Allocate any LocalAddrs if they haven't been allocated yet.
@@ -921,26 +865,14 @@ void UIpNetDriver::SetSocketAndLocalAddress(const TSharedPtr<FSocket>& SharedSoc
 
 void UIpNetDriver::ClearSockets()
 {
-	// For backwards compatability with the public Socket member. Destroy it manually if it won't be destroyed by the reset below.
-	PRAGMA_DISABLE_DEPRECATION_WARNINGS
-	if(!ensureMsgf(Socket == SocketPrivate.Get(), TEXT("UIpNetDriver::ClearSockets: Socket and SocketPrivate point to different sockets! %s"), *GetDescription()))
-	{
-		ISocketSubsystem* const SocketSubsystem = GetSocketSubsystem();
-
-		if (SocketSubsystem)
-		{
-			SocketSubsystem->DestroySocket(Socket);
-		}
-	}
-	Socket = nullptr;
-	PRAGMA_ENABLE_DEPRECATION_WARNINGS
-
 	SocketPrivate.Reset();
-	BoundSockets.Reset();
+	Resolver->ClearSockets();
 }
 
 bool UIpNetDriver::InitBase( bool bInitAsClient, FNetworkNotify* InNotify, const FURL& URL, bool bReuseAddressAndPort, FString& Error )
 {
+	using namespace UE::Net::Private;
+
 	if (!Super::InitBase(bInitAsClient, InNotify, URL, bReuseAddressAndPort, Error))
 	{
 		return false;
@@ -958,63 +890,24 @@ bool UIpNetDriver::InitBase( bool bInitAsClient, FNetworkNotify* InNotify, const
 	// and thus we rely on the OS socket to buffer a lot of data.
 	const int32 DesiredRecvSize = bInitAsClient ? ClientDesiredSocketReceiveBufferBytes : ServerDesiredSocketReceiveBufferBytes;
 	const int32 DesiredSendSize = bInitAsClient ? ClientDesiredSocketSendBufferBytes : ServerDesiredSocketSendBufferBytes;
-
-	TArray<TSharedRef<FInternetAddr>> BindAddresses = SocketSubsystem->GetLocalBindAddresses();
-
-	// Handle potentially empty arrays
-	if (BindAddresses.Num() == 0)
-	{
-		Error = TEXT("No binding addresses could be found or grabbed for this platform! Sockets could not be created!");
-		UE_LOG(LogNet, Error, TEXT("%s"), *Error);
-		return false;
-	}
-
-	// Create sockets for every bind address
-	for (TSharedRef<FInternetAddr>& BindAddr : BindAddresses)
-	{
-		FUniqueSocket NewSocket = CreateAndBindSocket(BindAddr, BindPort, bReuseAddressAndPort, DesiredRecvSize, DesiredSendSize, Error);
-		if (NewSocket.IsValid())
-		{
-			UE_LOG(LogNet, Log, TEXT("Created socket for bind address: %s on port %d"), *BindAddr->ToString(false), BindPort);
-			BoundSockets.Emplace(NewSocket.Release(), FSocketDeleter(NewSocket.GetDeleter()));
-		}
-		else
+	const EInitBindSocketsFlags InitBindFlags = bInitAsClient ? EInitBindSocketsFlags::Client : EInitBindSocketsFlags::Server;
+	FCreateAndBindSocketFunc CreateAndBindSocketsFunc = [this, BindPort, bReuseAddressAndPort, DesiredRecvSize, DesiredSendSize]
+									(TSharedRef<FInternetAddr> BindAddr, FString& Error) -> FUniqueSocket
 		{
-			UE_LOG(LogNet, Warning, TEXT("Could not create socket for bind address %s, got error %s"), *BindAddr->ToString(false), *Error);
-			Error = TEXT("");
-			continue;
-		}
+			return this->CreateAndBindSocket(BindAddr, BindPort, bReuseAddressAndPort, DesiredRecvSize, DesiredSendSize, Error);
+		};
 
-		// Servers should only have one socket that they bind on in our code.
-		if (!bInitAsClient)
-		{
-			break;
-		}
-	}
+	bool bInitBindSocketsSuccess = Resolver->InitBindSockets(MoveTemp(CreateAndBindSocketsFunc), InitBindFlags, SocketSubsystem, Error);
 
-	if (!Error.IsEmpty() || BoundSockets.Num() == 0)
+	if (!bInitBindSocketsSuccess)
 	{
-		UE_LOG(LogNet, Warning, TEXT("Encountered an error while creating sockets for the bind addresses. %s"), *Error);
-		
-		// Make sure to destroy all sockets that we don't end up using.
-		BoundSockets.Reset();
+		UE_LOG(LogNet, Error, TEXT("InitBindSockets failed: %s"), ToCStr(Error));
 
 		return false;
 	}
 
-	PRAGMA_DISABLE_DEPRECATION_WARNINGS
-	// Some derived drivers might have already set a socket, so don't override their values
-	if (Socket == nullptr)
-	{
-		// However if they haven't set a socket, go ahead and set one now.
-		SetSocketAndLocalAddress(BoundSockets[0]);
-	}
-	else if (!LocalAddr.IsValid()) // If they have set the socket but not the LocalAddr, do so now.
-	{
-		LocalAddr = SocketSubsystem->CreateInternetAddr();
-		Socket->GetAddress(*LocalAddr);
-	}
-	PRAGMA_ENABLE_DEPRECATION_WARNINGS
+
+	SetSocketAndLocalAddress(Resolver->GetFirstSocket());
 	
 	// If the cvar is set and the socket subsystem supports it, create the receive thread.
 	if (CVarNetIpNetDriverUseReceiveThread.GetValueOnAnyThread() != 0 && SocketSubsystem->IsSocketWaitSupported())
@@ -1036,11 +929,7 @@ bool UIpNetDriver::InitBase( bool bInitAsClient, FNetworkNotify* InNotify, const
 
 			if (bRetrieveTimestamps)
 			{
-				// Properly set this flag for every socket for each bind address.
-				for (TSharedPtr<FSocket>& SubSocket : BoundSockets)
-				{
-					SubSocket->SetRetrieveTimestamp(true);
-				}
+				Resolver->SetRetrieveTimestamp(true);
 			}
 
 			ERecvMultiFlags RecvMultiFlags = bRetrieveTimestamps ? ERecvMultiFlags::RetrieveTimestamps : ERecvMultiFlags::None;
@@ -1086,86 +975,12 @@ bool UIpNetDriver::InitConnect( FNetworkNotify* InNotify, const FURL& ConnectURL
 
 	// Create new connection.
 	ServerConnection = NewObject<UNetConnection>(GetTransientPackage(), NetConnectionClass);
-	UIpConnection* IPConnection = CastChecked<UIpConnection>(ServerConnection);
 
-	if (IPConnection == nullptr)
-	{
-		Error = TEXT("Could not cast the ServerConnection into the base connection class for this netdriver!");
-		return false;
-	}
-
-	PRAGMA_DISABLE_DEPRECATION_WARNINGS
-	ServerConnection->InitLocalConnection(this, Socket, ConnectURL, USOCK_Pending);
-	PRAGMA_ENABLE_DEPRECATION_WARNINGS
-	const bool bResolutionEnabled = FIpConnectionHelper::IsAddressResolutionEnabledForConnection(IPConnection);
-
-	int32 DestinationPort = ConnectURL.Port;
-	if (bResolutionEnabled)
-	{
-		FIpConnectionHelper::PushSocketsToConnection(IPConnection, BoundSockets);
-		BoundSockets.Empty();
-
-		// Create a weakobj so that we can pass the Connection safely to the lambda for later
-		TWeakObjectPtr<UIpConnection> SafeConnectionPtr(IPConnection);
+	ServerConnection->InitLocalConnection(this, SocketPrivate.Get(), ConnectURL, USOCK_Pending);
 
-		auto AsyncResolverHandler = [SafeConnectionPtr, SocketSubsystem, DestinationPort](FAddressInfoResult Results)
-		{
-			FGCScopeGuard Guard;
-
-			// Check if we still have a valid pointer
-			if (!SafeConnectionPtr.IsValid())
-			{
-				// If we got in here, we are already in some sort of exiting state typically.
-				// We shouldn't have to do any more other than not do any sort of operations on the connection
-				UE_LOG(LogNet, Warning, TEXT("GAI Resolver Lambda: The NetConnection class has become invalid after results for %s were grabbed."), *Results.QueryHostName);
-				return;
-			}
-			
-			if (Results.ReturnCode == SE_NO_ERROR)
-			{
-				TArray<TSharedRef<FInternetAddr>> AddressResults;
-				for (auto& Result : Results.Results)
-				{
-					AddressResults.Add(Result.Address);
-				}
-
-#if !UE_BUILD_SHIPPING
-				// This is useful for injecting a good result into the array to test the resolution system
-				const FString DebugAddressAddition = CVarNetDebugAddResolverAddress.GetValueOnAnyThread();
-				if (!DebugAddressAddition.IsEmpty())
-				{
-					TSharedPtr<FInternetAddr> SpecialResultAddr = SocketSubsystem->GetAddressFromString(DebugAddressAddition);
-					if (SpecialResultAddr.IsValid())
-					{
-						SpecialResultAddr->SetPort(DestinationPort);
-						AddressResults.Add(SpecialResultAddr.ToSharedRef());
-						UE_LOG(LogNet, Log, TEXT("Added additional result address %s to resolver list"), *SpecialResultAddr->ToString(false));
-					}
-				}
-#endif
-				FIpConnectionHelper::PushResolverResultsToConnection(SafeConnectionPtr.Get(), AddressResults);
-			}
-			else
-			{
-				FIpConnectionHelper::HandleResolverError(SafeConnectionPtr.Get());
-			}
-		};
-
-		SocketSubsystem->GetAddressInfoAsync(AsyncResolverHandler, *ConnectURL.Host, *FString::Printf(TEXT("%d"), DestinationPort),
-			EAddressInfoFlags::AllResultsWithMapping | EAddressInfoFlags::OnlyUsableAddresses, NAME_None, ESocketType::SOCKTYPE_Datagram);
-	}
-	else if (BoundSockets.Num() > 1)
-	{
-		// Clean up any potential multiple sockets we have created when resolution was disabled.
-		// InitBase could have created multiple sockets and if so, we'll want to clean them up.
-		UE_LOG(LogNet, Verbose, TEXT("Cleaning up additional sockets created as address resolution is disabled."));
-		BoundSockets.RemoveAll([InSocket = GetSocket()](const TSharedPtr<FSocket>& CurSocket)
-		{
-			return CurSocket.Get() != InSocket;
-		});
-	}
+	Resolver->InitConnect(ServerConnection, SocketSubsystem, GetSocket(), ConnectURL);
 	
-	UE_LOG(LogNet, Log, TEXT("Game client on port %i, rate %i"), DestinationPort, ServerConnection->CurrentNetSpeed );
+	UE_LOG(LogNet, Log, TEXT("Game client on port %i, rate %i"), ConnectURL.Port, ServerConnection->CurrentNetSpeed );
 	CreateInitialClientChannels();
 
 	return true;
@@ -1455,15 +1270,16 @@ void UIpNetDriver::TickDispatch(float DeltaTime)
 
 FSocket* UIpNetDriver::GetSocket()
 {
+	using namespace UE::Net::Private;
+
 	UIpConnection* IpServerConnection = Cast<UIpConnection>(ServerConnection);
-	if (FIpConnectionHelper::IsAddressResolutionEnabledForConnection(IpServerConnection))
+
+	if (IpServerConnection != nullptr && FNetDriverAddressResolution::GetConnectionResolver(IpServerConnection)->IsAddressResolutionEnabled())
 	{
-		return IpServerConnection->Socket;
+		return IpServerConnection->GetSocket();
 	}
 
-	PRAGMA_DISABLE_DEPRECATION_WARNINGS
-	return Socket;
-	PRAGMA_ENABLE_DEPRECATION_WARNINGS
+	return SocketPrivate.Get();
 }
 
 UNetConnection* UIpNetDriver::ProcessConnectionlessPacket(FReceivedPacketView& PacketRef, const FPacketBufferView& WorkingBuffer)
@@ -1600,9 +1416,7 @@ UNetConnection* UIpNetDriver::ProcessConnectionlessPacket(FReceivedPacketView& P
 			ReturnVal = NewObject<UIpConnection>(GetTransientPackage(), NetConnectionClass);
 			check(ReturnVal != nullptr);
 
-			PRAGMA_DISABLE_DEPRECATION_WARNINGS
-			ReturnVal->InitRemoteConnection(this, Socket, World ? World->URL : FURL(), * Address, USOCK_Open);
-			PRAGMA_ENABLE_DEPRECATION_WARNINGS
+			ReturnVal->InitRemoteConnection(this, SocketPrivate.Get(), World ? World->URL : FURL(), *Address, USOCK_Open);
 
 #if STATELESSCONNECT_HAS_RANDOM_SEQUENCE
 			// Set the initial packet sequence from the handshake data
@@ -1711,6 +1525,8 @@ FString UIpNetDriver::LowLevelGetNetworkNumber()
 
 void UIpNetDriver::LowLevelDestroy()
 {
+	using namespace UE::Net::Private;
+
 	Super::LowLevelDestroy();
 
 	// Close the socket.
@@ -1749,9 +1565,12 @@ void UIpNetDriver::LowLevelDestroy()
 			UE_LOG(LogExit, Log, TEXT("closesocket error (%i)"), (int32)SocketSubsystem->GetLastErrorCode() );
 		}
 
-		if (FIpConnectionHelper::IsAddressResolutionEnabledForConnection(IpServerConnection))
+		FNetConnectionAddressResolution* ConnResolver = FNetDriverAddressResolution::GetConnectionResolver(IpServerConnection);
+
+		if (ConnResolver->IsAddressResolutionEnabled())
 		{
-			FIpConnectionHelper::CleanUpConnectionSockets(IpServerConnection);
+			FIpConnectionHelper::CleanupDeprecatedSocket(IpServerConnection);
+			ConnResolver->CleanupResolutionSockets();
 		}
 
 		ClearSockets();
@@ -2012,6 +1831,8 @@ bool UIpNetDriver::FReceiveThreadRunnable::DispatchPacket(FReceivedPacket&& Inco
 
 uint32 UIpNetDriver::FReceiveThreadRunnable::Run()
 {
+	using namespace UE::Net::Private;
+
 	const FTimespan Timeout = FTimespan::FromMilliseconds(CVarNetIpNetDriverReceiveThreadPollTimeMS.GetValueOnAnyThread());
 	const float SleepTimeForWaitableErrorsInSec = CVarRcvThreadSleepTimeForWaitableErrorsInSeconds.GetValueOnAnyThread();
 	const int32 ActionForLongRecvErrors = CVarRcvThreadShouldSleepForLongRecvErrors.GetValueOnAnyThread();
@@ -2023,7 +1844,9 @@ uint32 UIpNetDriver::FReceiveThreadRunnable::Run()
 	{
 		// If we've encountered any errors during address resolution (this flag will not have the error state on it if resolution is disabled)
 		// Then stop running this thread. This stomps out any potential infinite loops caused by undefined behavior.
-		if (FIpConnectionHelper::HasAddressResolutionFailedForConnection(OwningNetDriver->GetServerConnection()))
+		UIpConnection* ServerConn = OwningNetDriver->GetServerConnection();
+
+		if (ServerConn != nullptr && FNetDriverAddressResolution::GetConnectionResolver(ServerConn)->HasAddressResolutionFailed())
 		{
 			break;
 		}
diff --git a/Engine/Plugins/Online/OnlineSubsystemUtils/Source/OnlineSubsystemUtils/Private/NetAddressResolution.cpp b/Engine/Plugins/Online/OnlineSubsystemUtils/Source/OnlineSubsystemUtils/Private/NetAddressResolution.cpp
new file mode 100644
--- /dev/null
+++ b/Engine/Plugins/Online/OnlineSubsystemUtils/Source/OnlineSubsystemUtils/Private/NetAddressResolution.cpp
@@ -0,0 +1,387 @@
+// Copyright Epic Games, Inc. All Rights Reserved.
+
+// Includes
+#include "NetAddressResolution.h"
+#include "SocketSubsystem.h"
+#include "IpNetDriver.h"
+#include "Sockets.h"
+
+
+// Stats
+DECLARE_CYCLE_STAT(TEXT("IpConnection Address Synthesis"), STAT_IpConnection_AddressSynthesis, STATGROUP_Net);
+
+
+// CVars
+#if !UE_BUILD_SHIPPING
+TAutoConsoleVariable<FString> CVarNetDebugAddResolverAddress(
+	TEXT("net.DebugAppendResolverAddress"),
+	TEXT(""),
+	TEXT("If this is set, all IP address resolution methods will add the value of this CVAR to the list of results.")
+		TEXT("This allows for testing resolution functionality across all multiple addresses with the end goal of having a successful result")
+		TEXT("(being the value of this CVAR)"),
+	ECVF_Default | ECVF_Cheat);
+#endif
+
+namespace UE::Net::Private
+{
+
+/**
+ * FNetDriverAddressResolution
+ */
+
+bool FNetDriverAddressResolution::InitBindSockets(FCreateAndBindSocketFunc CreateAndBindSocketFunc, EInitBindSocketsFlags Flags,
+													ISocketSubsystem* SocketSubsystem, FString& Error)
+{
+	TArray<TSharedRef<FInternetAddr>> BindAddresses = SocketSubsystem->GetLocalBindAddresses();
+
+	// Handle potentially empty arrays
+	if (BindAddresses.Num() == 0)
+	{
+		Error = TEXT("No binding addresses could be found or grabbed for this platform! Sockets could not be created!");
+		return false;
+	}
+
+	// Create sockets for every bind address
+	for (TSharedRef<FInternetAddr>& BindAddr : BindAddresses)
+	{
+		FUniqueSocket NewSocket = CreateAndBindSocketFunc(BindAddr, Error);
+
+		if (NewSocket.IsValid())
+		{
+			UE_LOG(LogNet, Log, TEXT("Created socket for bind address: %s"), ToCStr(BindAddr->ToString(true)));
+
+			BoundSockets.Emplace(NewSocket.Release(), FSocketDeleter(NewSocket.GetDeleter()));
+		}
+		else
+		{
+			UE_LOG(LogNet, Warning, TEXT("Could not create socket for bind address %s, got error %s"), ToCStr(BindAddr->ToString(false)),
+					ToCStr(Error));
+
+			Error = TEXT("");
+
+			continue;
+		}
+
+		// Servers should only have one socket that they bind on in our code.
+		if (EnumHasAnyFlags(Flags, EInitBindSocketsFlags::Server))
+		{
+			break;
+		}
+	}
+
+	if (!Error.IsEmpty() || BoundSockets.Num() == 0)
+	{
+		UE_LOG(LogNet, Warning, TEXT("Encountered an error while creating sockets for the bind addresses. %s"), *Error);
+		
+		// Make sure to destroy all sockets that we don't end up using.
+		BoundSockets.Reset();
+
+		return false;
+	}
+
+	return true;
+}
+
+void FNetDriverAddressResolution::InitConnect(UNetConnection* ServerConnection, ISocketSubsystem* SocketSubsystem, const FSocket* ActiveSocket,
+												const FURL& ConnectURL)
+{
+	UIpConnection* IPConnection = CastChecked<UIpConnection>(ServerConnection);
+	int32 DestinationPort = ConnectURL.Port;
+
+	if (IPConnection->Resolver->IsAddressResolutionEnabled())
+	{
+		IPConnection->Resolver->BindSockets = MoveTemp(BoundSockets);
+
+		// Create a weakobj so that we can pass the Connection safely to the lambda for later
+		TWeakObjectPtr<UIpConnection> SafeConnectionPtr(IPConnection);
+
+		auto AsyncResolverHandler = [SafeConnectionPtr, SocketSubsystem, DestinationPort](FAddressInfoResult Results)
+		{
+			FGCScopeGuard Guard;
+
+			// Check if we still have a valid pointer
+			if (!SafeConnectionPtr.IsValid())
+			{
+				// If we got in here, we are already in some sort of exiting state typically.
+				// We shouldn't have to do any more other than not do any sort of operations on the connection
+				UE_LOG(LogNet, Warning, TEXT("GAI Resolver Lambda: The NetConnection class has become invalid after results for %s were grabbed."), *Results.QueryHostName);
+				return;
+			}
+
+
+			UIpConnection* Connection = SafeConnectionPtr.Get();
+			
+			if (Results.ReturnCode == SE_NO_ERROR)
+			{
+				TArray<TSharedRef<FInternetAddr>> AddressResults;
+				for (auto& Result : Results.Results)
+				{
+					AddressResults.Add(Result.Address);
+				}
+
+#if !UE_BUILD_SHIPPING
+				// This is useful for injecting a good result into the array to test the resolution system
+				const FString DebugAddressAddition = CVarNetDebugAddResolverAddress.GetValueOnAnyThread();
+				if (!DebugAddressAddition.IsEmpty())
+				{
+					TSharedPtr<FInternetAddr> SpecialResultAddr = SocketSubsystem->GetAddressFromString(DebugAddressAddition);
+					if (SpecialResultAddr.IsValid())
+					{
+						SpecialResultAddr->SetPort(DestinationPort);
+						AddressResults.Add(SpecialResultAddr.ToSharedRef());
+						UE_LOG(LogNet, Log, TEXT("Added additional result address %s to resolver list"), *SpecialResultAddr->ToString(false));
+					}
+				}
+#endif
+
+				Connection->Resolver->ResolverResults = MoveTemp(AddressResults);
+				Connection->Resolver->ResolutionState = EAddressResolutionState::TryNextAddress;
+			}
+			else
+			{
+				Connection->Resolver->ResolutionState = EAddressResolutionState::Error;
+
+				Connection->Close(ENetCloseResult::AddressResolutionFailed);
+			}
+		};
+
+		SocketSubsystem->GetAddressInfoAsync(AsyncResolverHandler, *ConnectURL.Host, *FString::Printf(TEXT("%d"), DestinationPort),
+			EAddressInfoFlags::AllResultsWithMapping | EAddressInfoFlags::OnlyUsableAddresses, NAME_None, ESocketType::SOCKTYPE_Datagram);
+	}
+	else
+	{
+		// Clean up any potential multiple sockets we have created when resolution was disabled.
+		// InitBase could have created multiple sockets and if so, we'll want to clean them up.
+		for (int32 SockIdx=BoundSockets.Num()-1; SockIdx >= 0; SockIdx--)
+		{
+			if (BoundSockets[SockIdx].Get() != ActiveSocket)
+			{
+				BoundSockets.RemoveAt(SockIdx);
+			}
+		}
+	}
+}
+
+void FNetDriverAddressResolution::SetRetrieveTimestamp(bool bRetrieveTimestamp)
+{
+	for (TSharedPtr<FSocket>& CurSocket : BoundSockets)
+	{
+		CurSocket->SetRetrieveTimestamp(bRetrieveTimestamp);
+	}
+}
+
+FNetConnectionAddressResolution* FNetDriverAddressResolution::GetConnectionResolver(UIpConnection* Connection)
+{
+	return Connection != nullptr ? Connection->Resolver.Get() : nullptr;
+}
+
+
+/**
+ * FNetConnectionAddressResolution
+ */
+
+FNetConnectionAddressResolution::FNetConnectionAddressResolution(const FSocket* const & InDeprecatedSocket)
+	: DeprecatedSocket(InDeprecatedSocket)
+	, ResolutionState(EAddressResolutionState::None)
+{
+}
+
+bool FNetConnectionAddressResolution::InitLocalConnection(ISocketSubsystem* SocketSubsystem, FSocket* InSocket, const FURL& InURL)
+{
+	bool bValidInit = true;
+
+	// If resolution is disabled, fall back to address synthesis
+	if (!IsAddressResolutionEnabled())
+	{
+		// Figure out IP address from the host URL
+		bValidInit = false;
+
+		// Get numerical address directly.
+		RemoteAddr = SocketSubsystem->CreateInternetAddr();
+		RemoteAddr->SetIp(*InURL.Host, bValidInit);
+
+		// If the protocols do not match, attempt to synthesize the address so they do.
+		FName SocketProtocol = InSocket->GetProtocol();
+
+		if ((bValidInit && SocketProtocol != RemoteAddr->GetProtocolType()) || !bValidInit)
+		{
+			SCOPE_CYCLE_COUNTER(STAT_IpConnection_AddressSynthesis);
+
+			// We want to use GAI to create the address with the correct protocol.
+			const FAddressInfoResult MapRequest = SocketSubsystem->GetAddressInfo(*InURL.Host, nullptr,
+				EAddressInfoFlags::AllResultsWithMapping | EAddressInfoFlags::OnlyUsableAddresses, SocketProtocol);
+
+			// Set the remote addr provided we have information.
+			if (MapRequest.ReturnCode == SE_NO_ERROR && MapRequest.Results.Num() > 0)
+			{
+				RemoteAddr = MapRequest.Results[0].Address->Clone();
+				bValidInit = true;
+			}
+			else
+			{
+				UE_LOG(LogNet, Warning, TEXT("IpConnection::InitConnection: Address protocols do not match and cannot be synthesized to a ")
+						TEXT("similar address, this will likely lead to issues!"));
+			}
+		}
+
+		RemoteAddr->SetPort(InURL.Port);
+
+		if (!bValidInit)
+		{
+			UE_LOG(LogNet, Verbose, TEXT("IpConnection::InitConnection: Unable to resolve %s"), ToCStr(InURL.Host));
+		}
+	}
+	else
+	{
+		ResolutionState = EAddressResolutionState::WaitingForResolves;
+	}
+
+	return bValidInit;
+}
+
+EEAddressResolutionHandleResult FNetConnectionAddressResolution::NotifyTimeout()
+{
+	EEAddressResolutionHandleResult Result = EEAddressResolutionHandleResult::CallerShouldHandle;
+
+	if (CanContinueResolution())
+	{
+		ResolutionState = EAddressResolutionState::TryNextAddress;
+		Result = EEAddressResolutionHandleResult::HandledInternally;
+	}
+
+	return Result;
+}
+
+EEAddressResolutionHandleResult FNetConnectionAddressResolution::NotifyReceiveError()
+{
+	EEAddressResolutionHandleResult Result = EEAddressResolutionHandleResult::CallerShouldHandle;
+
+	if (CanContinueResolution())
+	{
+		ResolutionState = EAddressResolutionState::TryNextAddress;
+		Result = EEAddressResolutionHandleResult::HandledInternally;
+	}
+	else
+	{
+		ResolutionState = EAddressResolutionState::Error;
+	}
+
+	return Result;
+}
+
+EEAddressResolutionHandleResult FNetConnectionAddressResolution::NotifySendError()
+{
+	EEAddressResolutionHandleResult Result = EEAddressResolutionHandleResult::CallerShouldHandle;
+
+	if (CanContinueResolution())
+	{
+		ResolutionState = EAddressResolutionState::TryNextAddress;
+		Result = EEAddressResolutionHandleResult::HandledInternally;
+	}
+	else
+	{
+		ResolutionState = EAddressResolutionState::Error;
+	}
+
+	return Result;
+}
+
+ECheckAddressResolutionResult FNetConnectionAddressResolution::CheckAddressResolution()
+{
+	ECheckAddressResolutionResult Result = ECheckAddressResolutionResult::None;
+
+	if (ResolutionState == EAddressResolutionState::TryNextAddress)
+	{
+		if (CurrentAddressIndex >= ResolverResults.Num())
+		{
+			UE_LOG(LogNet, Warning, TEXT("Exhausted the number of resolver results, closing the connection now."));
+			ResolutionState = EAddressResolutionState::Error;
+			return ECheckAddressResolutionResult::None;
+		}
+
+		RemoteAddr = ResolverResults[CurrentAddressIndex];
+
+		ResolutionSocket.Reset();
+
+		for (const TSharedPtr<FSocket>& BindSocket : BindSockets)
+		{
+			if (BindSocket->GetProtocol() == RemoteAddr->GetProtocolType())
+			{
+				ResolutionSocket = BindSocket;
+				break;
+			}
+		}
+
+		if (ResolutionSocket.IsValid())
+		{
+			ResolutionState = EAddressResolutionState::Connecting;
+
+			if (CurrentAddressIndex == 0)
+			{
+				Result = ECheckAddressResolutionResult::TryFirstAddress;
+			}
+			else
+			{
+				Result = ECheckAddressResolutionResult::TryNextAddress;
+			}
+
+			++CurrentAddressIndex;
+		}
+		else
+		{
+			UE_LOG(LogNet, Error, TEXT("Unable to find a binding socket for the resolve address result %s"), ToCStr(RemoteAddr->ToString(true)));
+
+			ResolutionState = EAddressResolutionState::Error;
+			Result = ECheckAddressResolutionResult::FindSocketError;
+		}
+	}
+	else if (ResolutionState == EAddressResolutionState::Connected)
+	{
+		ResolutionState = EAddressResolutionState::Done;
+		Result = ECheckAddressResolutionResult::Connected;
+
+		CleanupResolutionSockets(ECleanupResolutionSocketsFlags::CleanInactive);
+	}
+	else if (ResolutionState == EAddressResolutionState::Error)
+	{
+		UE_LOG(LogNet, Warning, TEXT("Encountered an error, cleaning up this connection now"));
+
+		ResolutionState = EAddressResolutionState::Done;
+		Result = ECheckAddressResolutionResult::Error;
+	}
+
+	return Result;
+}
+
+void FNetConnectionAddressResolution::NotifyAddressResolutionConnected()
+{
+#if !UE_BUILD_SHIPPING
+	checkf(ResolutionState == EAddressResolutionState::Connecting,
+				TEXT("NotifyAddressResolutionConnected should only occur during 'Connecting' ResolutionState."));
+#endif
+
+	ResolutionState = EAddressResolutionState::Connected;
+}
+
+void FNetConnectionAddressResolution::CleanupResolutionSockets(ECleanupResolutionSocketsFlags CleanupFlags
+																/*=ECleanupResolutionSocketsFlags::CleanAll*/)
+{
+	if (IsAddressResolutionEnabled())
+	{
+		if (EnumHasAnyFlags(CleanupFlags, ECleanupResolutionSocketsFlags::CleanAll))
+		{
+			// Remove when UIpConnection.Socket is deprecated
+#if !UE_BUILD_SHIPPING
+			ensureMsgf(DeprecatedSocket == nullptr,
+						TEXT("All codepaths passing 'CleanAll', must call 'UIpConnection::CleanupDeprecatedSocket' first."));
+#endif
+
+			ResolutionSocket.Reset();
+		}
+	
+		BindSockets.Reset();
+		ResolverResults.Empty();
+	}
+}
+
+}
diff --git a/Engine/Plugins/Online/OnlineSubsystemUtils/Source/OnlineSubsystemUtils/Private/NetAddressResolution.h b/Engine/Plugins/Online/OnlineSubsystemUtils/Source/OnlineSubsystemUtils/Private/NetAddressResolution.h
new file mode 100644
--- /dev/null
+++ b/Engine/Plugins/Online/OnlineSubsystemUtils/Source/OnlineSubsystemUtils/Private/NetAddressResolution.h
@@ -0,0 +1,352 @@
+// Copyright Epic Games, Inc. All Rights Reserved.
+
+#pragma once
+
+
+// Includes
+#include "Misc/EnumClassFlags.h"
+#include "SocketSubsystem.h"
+#include "Templates/Function.h"
+#include "Templates/SharedPointer.h"
+#include "Containers/Array.h"
+#include "Containers/UnrealString.h"
+#include "Engine/EngineBaseTypes.h"
+#include "IpConnection.h"
+
+
+// Forward declarations
+class ISocketSubsystem;
+class FSocket;
+class FInternetAddr;
+class UNetConnection;
+class UIpNetDriver;
+
+namespace UE::Net::Private
+{
+
+class FNetConnectionAddressResolution;
+
+
+/**
+ * NetDriver callback creating and binding the NetDriver-specific socket
+ *
+ * @param BindAddr	The bind address for the socket
+ * @param Error		If there was an error, specifies the error message
+ */
+using FCreateAndBindSocketFunc = TUniqueFunction<FUniqueSocket(TSharedRef<FInternetAddr> /*BindAddr*/, FString& /*Error*/)>;
+
+/**
+ * Flags specifying the type of call to InitBindSockets
+ */
+enum class EInitBindSocketsFlags : uint16
+{
+	Client		= 0x0001,
+	Server		= 0x0002
+};
+
+ENUM_CLASS_FLAGS(EInitBindSocketsFlags);
+
+
+/**
+ * NetDriver level Address Resolution
+ */
+class FNetDriverAddressResolution
+{
+	friend UIpNetDriver;
+
+private:
+	/**
+	 * Creates and binds sockets for each local bind address
+	 *
+	 * @param CreateAndBindSocketFunc	The NetDriver callback for creating the actual socket
+	 * @param Flags						Flags specifying the type of initialization this is (typically Client vs Server)
+	 * @param SocketSubsystem			The NetDriver socket subsystem
+	 * @param Error						If there was an error, specifies the error message
+	 * @return							Whether or not the call was successful
+	 */
+	bool InitBindSockets(FCreateAndBindSocketFunc CreateAndBindSocketFunc, EInitBindSocketsFlags Flags, ISocketSubsystem* SocketSubsystem,
+							FString& Error);
+
+	/**
+	 * Kicks off NetDriver level address resolution for clients connecting to servers (with resolution passing on to FNetConnectionAddressResolution).
+	 *
+	 * @param ServerConnection		The connection to the server
+	 * @param SocketSubsystem		The NetDriver socket subsystem
+	 * @param ActiveSocket			Read-only pointer to the active NetDriver socket - in case it is not controlled by address resolution
+	 * @param ConnectURL			The URL used to connect to the server
+	 */
+	void InitConnect(UNetConnection* ServerConnection, ISocketSubsystem* SocketSubsystem, const FSocket* ActiveSocket, const FURL& ConnectURL);
+
+
+	/** Clears all references to sockets held by this resolver. */
+	void ClearSockets()
+	{
+		BoundSockets.Reset();
+	}
+
+	/**
+	 * Gets the first socket bound by address resolution.
+	 *
+	 * @return	The first socket bound by address resolution
+	 */
+	TSharedPtr<FSocket> GetFirstSocket() const
+	{
+		return BoundSockets.Num() > 0 ? BoundSockets[0] : TSharedPtr<FSocket>();
+	}
+
+	/**
+	 * Sets the internal socket flag for whether or not timestamps from received packets should be retrieved.
+	 *
+	 * @param bRetrieveTimestamp	Whether or not timestamps should be retrieved with received packets
+	 */
+	void SetRetrieveTimestamp(bool bRetrieveTimestamp);
+
+	/**
+	 * Gets the NetConnection specific address resolver, for a connection
+	 *
+	 * @param Connection			The connection whose Resolver we want access to
+	 * @return						Returns a pointer to the connection's Resolver
+	 */
+	static FNetConnectionAddressResolution* GetConnectionResolver(UIpConnection* Connection);
+
+private:
+	/**
+	 * An array sockets created for every binding address a machine has in use for performing address resolution.
+	 * This array empties after connections have been spun up.
+	 */
+	TArray<TSharedPtr<FSocket>> BoundSockets;
+};
+
+
+/** A state system of the address resolution functionality. */
+enum class EAddressResolutionState : uint8
+{
+	None,					// Undefined
+	Disabled,				// Address resolution is explicitly disabled
+	WaitingForResolves,		// Address resolution is underway and we're awaiting results
+	Connecting,				// Address resolution for a single address is complete and we're attempting to connect (not complete yet)
+	TryNextAddress,			// Address resolution for the current address failed, and we're trying the next address on the list
+	Connected,				// We have successfully connected to the current address - awaiting completion on next state check
+	Done,					// Address resolution has completed
+	Error					// There was an error during address resolution
+};
+
+/**
+ * Result for CheckAddressResolution - partially mirrors EAddressResolutionState
+ */
+enum class ECheckAddressResolutionResult : uint8
+{
+	None,				// No action needed
+	TryFirstAddress,	// Trying the first address for address resolution
+	TryNextAddress,		// Trying a subsequent address for address resolution
+	Connected,			// Address resolution was just marked as connected to the current address
+	Error,				// There was an error during address resolution
+	FindSocketError		// Unable to find a socket that can be used with the current resolve address
+};
+
+/**
+ * Specifies whether an Address Resolution 'Handle*' function handled the event internally, or whether it should be handled by the caller instead
+ */
+enum class EEAddressResolutionHandleResult : uint8
+{
+	HandledInternally,		// The event was handled within the internal address resolution code, and needs no handling by the caller
+	CallerShouldHandle		// The event was unhandled, the caller of the 'Handle*' function should handle the event
+};
+
+/**
+ * CleanupResolutionSockets parameter flags (implemented as flags for extensibility/multiple flags)
+ */
+enum class ECleanupResolutionSocketsFlags : uint8
+{
+	CleanAll		= 0x01,	// Clean all sockets, including the active one
+	CleanInactive	= 0x02	// Clean only inactive sockets
+};
+
+ENUM_CLASS_FLAGS(ECleanupResolutionSocketsFlags);
+
+
+/**
+ * NetConnection level Address Resolution
+ *
+ * NOTE: Private interface is accessible to the NetConnection, the public interface is only for the NetDriver
+ */
+class FNetConnectionAddressResolution
+{
+	friend UIpConnection;
+	friend FNetDriverAddressResolution;
+
+public:
+	/**
+	 * Base constructor
+	 *
+	 * @param InDeprecatedSocket	Read-only reference to deprecated UIpConnection.Socket variable, for verification checks
+	 */
+	FNetConnectionAddressResolution(const FSocket* const & InDeprecatedSocket);
+
+	/**
+	 * Checks to see if the owner NetConnection class can use address resolution
+	 *
+	 * @return if address resolution is allowed to continue processing.
+	 */
+	bool IsAddressResolutionEnabled() const
+	{
+		return ResolutionState != EAddressResolutionState::Disabled;
+	}
+
+	/**
+	 * Whether or not address resolution has failed
+	 *
+	 * @return	Returns whether or not address resolution has failed
+	 */
+	bool HasAddressResolutionFailed() const
+	{
+		return ResolutionState == EAddressResolutionState::Error;
+	}
+
+	/**
+	 * Cleans up the socket information in use with resolution. This can get called numerous times.
+	 *
+	 * @param CleanupFlags		Flags specifying what should be cleaned up
+	 */
+	void CleanupResolutionSockets(ECleanupResolutionSocketsFlags CleanupFlags=ECleanupResolutionSocketsFlags::CleanAll);
+
+private:
+	/**
+	 * Initializes Address resolution for clients connecting to a server.
+	 *
+	 * @param SocketSubsystem	The NetConnection socket subsystem
+	 * @param InSocket			When address resolution is disabled, the socket being used for connecting.
+	 * @param InURL				The URL the client is connecting to
+	 * @return					Whether or not address resolution kicked off successfully
+	 */
+	bool InitLocalConnection(ISocketSubsystem* SocketSubsystem, FSocket* InSocket, const FURL& InURL);
+
+	/**
+	 * Gives address resolution a chance to override a NetConnection timeout.
+	 *
+	 * @return		Whether or not the timeout was handled internally by address resolution, or whether it should be handled by the NetConnection
+	 */
+	EEAddressResolutionHandleResult NotifyTimeout();
+
+	/**
+	 * Gives address resolution a chance to override NetConnection handling of packet/socket receive errors.
+	 *
+	 * @return		Whether or not the error was handled internally by address resolution, or whether it should be handled by the NetConnection
+	 */
+	EEAddressResolutionHandleResult NotifyReceiveError();
+
+	/**
+	 * Gives address resolution a chance to override NetConnection handling of socket send errors.
+	 *
+	 * @return		Whether or not the error was handled internally by address resolution, or whether it should be handled by the NetConnection
+	 */
+	EEAddressResolutionHandleResult NotifySendError();
+
+	/**
+	 * Checks the state of address resolution, and kicks off next stage if applicable - returning the result/action taken or current status update.
+	 *
+	 * @return		Returns the result/action or the current status update.
+	 */
+	ECheckAddressResolutionResult CheckAddressResolution();
+
+	/**
+	 * Notification from the NetConnection that address resolution has connected (that a packet was sucessfully received)
+	 */
+	void NotifyAddressResolutionConnected();
+
+	/**
+	 * Determines if we can continue processing resolution results or not based on flags and current flow.
+	 *
+	 * @return if resolution is allowed to continue processing.
+	 */
+	bool CanContinueResolution() const {
+		return CurrentAddressIndex < ResolverResults.Num() && IsAddressResolutionEnabled() &&
+			ResolutionState != EAddressResolutionState::Error && ResolutionState != EAddressResolutionState::Done;
+	}
+
+	/**
+	 * Disables address resolution by pushing the disabled flag into the status field.
+	 */
+	void DisableAddressResolution()
+	{
+		ResolutionState = EAddressResolutionState::Disabled;
+	}
+
+	/**
+	 * Whether or not address resolution is in the process of attempting to connect to the server using the current address
+	 *
+	 * @return		Whether or not a connection attempt is underway
+	 */
+	bool IsAddressResolutionConnecting() const
+	{
+		return ResolutionState == EAddressResolutionState::Connecting;
+	}
+
+	/**
+	 * Whether or not address resolution has completed.
+	 *
+	 * @return		Whether or not address resolution has completed
+	 */
+	bool IsAddressResolutionComplete() const
+	{
+		return ResolutionState == EAddressResolutionState::Done;
+	}
+
+	/**
+	 * Whether or not address resolution is requesting that packet sends should be blocked.
+	 *
+	 * @return		Whether or not to block packet sends
+	 */
+	bool ShouldBlockSend() const
+	{
+		return IsAddressResolutionEnabled() &&
+			(ResolutionState == EAddressResolutionState::WaitingForResolves || ResolutionState == EAddressResolutionState::TryNextAddress);
+	}
+
+	/**
+	 * Gets the currently resolved socket (used as the primary NetConnection socket).
+	 *
+	 * @return		Returns the currently resolved socket
+	 */
+	TSharedPtr<FSocket> GetResolutionSocket() const
+	{
+		return ResolutionSocket;
+	}
+
+	/**
+	 * Gets the remote address of the currently resolved socket.
+	 *
+	 * @return		The remote address of the currently resolved socket.
+	 */
+	TSharedPtr<FInternetAddr> GetRemoteAddr() const
+	{
+		return RemoteAddr;
+	}
+
+
+private:
+	/** Temporary read-only reference to the deprecated UIpConnection.Socket value, for verification checks */
+	const FSocket* const& DeprecatedSocket;
+
+	/** An array of sockets tied to every binding address. */
+	TArray<TSharedPtr<FSocket>> BindSockets;
+
+	/** Holds a refcount to the actual socket to be used from BindSockets. */
+	TSharedPtr<FSocket> ResolutionSocket;
+
+	/** Stores the remote address from socket resolution */
+	TSharedPtr<FInternetAddr> RemoteAddr;
+
+	/** An array containing the address results GAI returns for the current host value. Given to us from the netdriver. */
+	TArray<TSharedRef<FInternetAddr>> ResolverResults;
+
+	/** The index into the ResolverResults that we're currently attempting */
+	int32 CurrentAddressIndex = 0;
+
+	/** 
+	 *  The connection's current status of where it is in the resolution state machine.
+	 *  If a platform should not use resolution, call DisableAddressResolution() in your constructor
+	 */
+	EAddressResolutionState ResolutionState;
+};
+}
diff --git a/Engine/Plugins/Online/SocketSubsystemEOS/Source/SocketSubsystemEOS/Private/NetConnectionEOS.cpp b/Engine/Plugins/Online/SocketSubsystemEOS/Source/SocketSubsystemEOS/Private/NetConnectionEOS.cpp
--- a/Engine/Plugins/Online/SocketSubsystemEOS/Source/SocketSubsystemEOS/Private/NetConnectionEOS.cpp
+++ b/Engine/Plugins/Online/SocketSubsystemEOS/Source/SocketSubsystemEOS/Private/NetConnectionEOS.cpp
@@ -56,21 +56,17 @@ void UNetConnectionEOS::CleanUp()
 
 void UNetConnectionEOS::DestroyEOSConnection()
 {
-	if (!Socket)
-	{
-		return;
-	}
+	FSocket* CurSocket = GetSocket();
 
-	if (!bHasP2PSession)
+	if (CurSocket != nullptr && bHasP2PSession)
 	{
-		return;
-	}
+		bHasP2PSession = false;
 
-	bHasP2PSession = false;
+		TSharedPtr<FInternetAddrEOS> RemoteAddrEOS = StaticCastSharedPtr<FInternetAddrEOS>(RemoteAddr);
 
-	TSharedPtr<FInternetAddrEOS> RemoteAddrEOS = StaticCastSharedPtr<FInternetAddrEOS>(RemoteAddr);
-	if (RemoteAddrEOS.IsValid())
-	{
-		static_cast<FSocketEOS*>(Socket)->Close(*RemoteAddrEOS);
+		if (RemoteAddrEOS.IsValid())
+		{
+			static_cast<FSocketEOS*>(CurSocket)->Close(*RemoteAddrEOS);
+		}
 	}
 }
diff --git a/Engine/Source/Runtime/Net/Core/Public/Net/Core/Connection/NetCloseResult.h b/Engine/Source/Runtime/Net/Core/Public/Net/Core/Connection/NetCloseResult.h
--- a/Engine/Source/Runtime/Net/Core/Public/Net/Core/Connection/NetCloseResult.h
+++ b/Engine/Source/Runtime/Net/Core/Public/Net/Core/Connection/NetCloseResult.h
@@ -101,6 +101,9 @@ enum class ENetCloseResult : uint32
 	/** JoinSplit failed */
 	JoinSplitFailure,
 
+	/** Address Rssolution Failed */
+	AddressResolutionFailed,
+
 
 	/** NetConnection Error result types */
 
