CIS fix.

diff --git a/Engine/Plugins/Online/OnlineSubsystemUtils/Source/OnlineSubsystemUtils/Private/IpNetDriver.cpp b/Engine/Plugins/Online/OnlineSubsystemUtils/Source/OnlineSubsystemUtils/Private/IpNetDriver.cpp
--- a/Engine/Plugins/Online/OnlineSubsystemUtils/Source/OnlineSubsystemUtils/Private/IpNetDriver.cpp
+++ b/Engine/Plugins/Online/OnlineSubsystemUtils/Source/OnlineSubsystemUtils/Private/IpNetDriver.cpp
@@ -1561,7 +1561,7 @@ void UIpNetDriver::LowLevelDestroy()
 
 		FNetConnectionAddressResolution* ConnResolver = FNetDriverAddressResolution::GetConnectionResolver(IpServerConnection);
 
-		if (ConnResolver->IsAddressResolutionEnabled())
+		if (ConnResolver != nullptr && ConnResolver->IsAddressResolutionEnabled())
 		{
 			FIpConnectionHelper::CleanupDeprecatedSocket(IpServerConnection);
 			ConnResolver->CleanupResolutionSockets();
