Further CIS fix.

diff --git a/Engine/Plugins/Online/OnlineSubsystemUtils/Source/OnlineSubsystemUtils/Private/NetAddressResolution.cpp b/Engine/Plugins/Online/OnlineSubsystemUtils/Source/OnlineSubsystemUtils/Private/NetAddressResolution.cpp
--- a/Engine/Plugins/Online/OnlineSubsystemUtils/Source/OnlineSubsystemUtils/Private/NetAddressResolution.cpp
+++ b/Engine/Plugins/Online/OnlineSubsystemUtils/Source/OnlineSubsystemUtils/Private/NetAddressResolution.cpp
@@ -355,12 +355,10 @@ ECheckAddressResolutionResult FNetConnectionAddressResolution::CheckAddressResol
 
 void FNetConnectionAddressResolution::NotifyAddressResolutionConnected()
 {
-#if !UE_BUILD_SHIPPING
-	checkf(ResolutionState == EAddressResolutionState::Connecting,
-				TEXT("NotifyAddressResolutionConnected should only occur during 'Connecting' ResolutionState."));
-#endif
-
-	ResolutionState = EAddressResolutionState::Connected;
+	if (ResolutionState == EAddressResolutionState::Connecting)
+	{
+		ResolutionState = EAddressResolutionState::Connected;
+	}
 }
 
 void FNetConnectionAddressResolution::CleanupResolutionSockets(ECleanupResolutionSocketsFlags CleanupFlags
