Non-unity fix.

diff --git a/Engine/Plugins/Online/OnlineSubsystemUtils/Source/OnlineSubsystemUtils/Classes/IpNetDriver.h b/Engine/Plugins/Online/OnlineSubsystemUtils/Source/OnlineSubsystemUtils/Classes/IpNetDriver.h
--- a/Engine/Plugins/Online/OnlineSubsystemUtils/Source/OnlineSubsystemUtils/Classes/IpNetDriver.h
+++ b/Engine/Plugins/Online/OnlineSubsystemUtils/Source/OnlineSubsystemUtils/Classes/IpNetDriver.h
@@ -13,6 +13,7 @@
 #include "Containers/CircularQueue.h"
 #include "SocketTypes.h"
 #include "SocketSubsystem.h"
+#include "Templates/PimplPtr.h"
 #include "IpNetDriver.generated.h"
 
 
