Initial connect diagnostics.

diff --git a/Engine/Plugins/Online/OnlineSubsystemUtils/Source/OnlineSubsystemUtils/Classes/IpConnection.h b/Engine/Plugins/Online/OnlineSubsystemUtils/Source/OnlineSubsystemUtils/Classes/IpConnection.h
--- a/Engine/Plugins/Online/OnlineSubsystemUtils/Source/OnlineSubsystemUtils/Classes/IpConnection.h
+++ b/Engine/Plugins/Online/OnlineSubsystemUtils/Source/OnlineSubsystemUtils/Classes/IpConnection.h
@@ -122,6 +122,17 @@ private:
 	/** Cached time of the first recv socket error that will be used to compute disconnect delay. */
 	double SocketError_RecvDelayStartTime;
 
+#if !UE_BUILD_SHIPPING
+	/** The number of socket-level sends that have occurred during initial connect */
+	int32 InitialConnectSocketSendCount = 0;
+
+	/** The last time initial connect diagnostics put out a log */
+	double InitialConnectLastLogTime = 0.0;
+
+	/** The value of 'InitialConnectSocketSendCount' the last time there was an initial connect diagnostic log */
+	int32 InitialConnectLastLogSocketSendCount = 0;
+#endif
+
 private:
 
 	/** Cleanup for the deprecated 'Socket' value */
diff --git a/Engine/Plugins/Online/OnlineSubsystemUtils/Source/OnlineSubsystemUtils/Private/IpConnection.cpp b/Engine/Plugins/Online/OnlineSubsystemUtils/Source/OnlineSubsystemUtils/Private/IpConnection.cpp
--- a/Engine/Plugins/Online/OnlineSubsystemUtils/Source/OnlineSubsystemUtils/Private/IpConnection.cpp
+++ b/Engine/Plugins/Online/OnlineSubsystemUtils/Source/OnlineSubsystemUtils/Private/IpConnection.cpp
@@ -46,6 +46,24 @@ TAutoConsoleVariable<int32> CVarNetIpConnectionDisableResolution(
 	TEXT("If enabled, any future ip connections will not use resolution methods."),
 	ECVF_Default | ECVF_Cheat);
 
+namespace UE::Net::Private
+{
+#if !UE_BUILD_SHIPPING
+	static int32 GNetDebugInitialConnect = 1;
+
+	static FAutoConsoleVariableRef CVarNetDebugInitialConnect(
+		TEXT("net.DebugInitialConnect"),
+		GNetDebugInitialConnect,
+		TEXT("When enabled, periodically logs socket-level send stats clientside, until a packet is successfully received to verify connection."));
+
+	static float GNetDebugInitialConnectLogFrequency = 10.0;
+
+	static FAutoConsoleVariableRef CVarNetDebugInitialConnectLogFrequency(
+		TEXT("net.DebugInitialConnectLogFrequency"),
+		GNetDebugInitialConnectLogFrequency,
+		TEXT("The amount of time, in seconds, between initial connect debug logging."));
+#endif
+}
 
 /**
  * UIpConnection
@@ -210,6 +228,25 @@ void UIpConnection::Tick(float DeltaSeconds)
 	}
 
 	Super::Tick(DeltaSeconds);
+
+#if !UE_BUILD_SHIPPING
+	if (!!GNetDebugInitialConnect && InTotalPackets == 0 && Driver != nullptr && Driver->ServerConnection != nullptr)
+	{
+		const double CurTime = Driver->GetElapsedTime();
+		const double ElapsedLogTime = CurTime - InitialConnectLastLogTime;
+
+		if (ElapsedLogTime >= static_cast<double>(GNetDebugInitialConnectLogFrequency))
+		{
+			const int32 SendCount = InitialConnectSocketSendCount - InitialConnectLastLogSocketSendCount;
+
+			UE_LOG(LogNet, Log, TEXT("Initial Connect Diagnostics: Sent '%i' packets in last '%f' seconds, no packets received yet."),
+					SendCount, ElapsedLogTime);
+
+			InitialConnectLastLogSocketSendCount = InitialConnectSocketSendCount;
+			InitialConnectLastLogTime = CurTime;
+		}
+	}
+#endif	
 }
 
 void UIpConnection::CleanUp()
@@ -264,6 +301,8 @@ void UIpConnection::WaitForSendTasks()
 
 void UIpConnection::LowLevelSend(void* Data, int32 CountBits, FOutPacketTraits& Traits)
 {
+	using namespace UE::Net::Private;
+
 	const uint8* DataToSend = reinterpret_cast<uint8*>(Data);
 
 	// If the remote addr hasn't been set, we need to wait for it.
@@ -369,6 +408,13 @@ void UIpConnection::LowLevelSend(void* Data, int32 CountBits, FOutPacketTraits&
 							bWasSendSuccessful = CurSocket->SendTo(Packet.GetData(), Packet.Num(), Result.BytesSent, *RemoteAddr);
 						}
 
+#if !UE_BUILD_SHIPPING
+						if (!!GNetDebugInitialConnect && InTotalPackets == 0 && bWasSendSuccessful)
+						{
+							InitialConnectSocketSendCount++;
+						}
+#endif
+
 						if (!bWasSendSuccessful && SocketSubsystem)
 						{
 							Result.Error = SocketSubsystem->GetLastErrorCode();
@@ -406,6 +452,13 @@ void UIpConnection::LowLevelSend(void* Data, int32 CountBits, FOutPacketTraits&
 					{
 						HandleSocketSendResult(SendResult, nullptr);
 					}
+
+#if !UE_BUILD_SHIPPING
+					if (!!GNetDebugInitialConnect && InTotalPackets == 0)
+					{
+						InitialConnectSocketSendCount++;
+					}
+#endif
 				}
 				else
 				{
diff --git a/Engine/Source/Runtime/Engine/Private/NetConnection.cpp b/Engine/Source/Runtime/Engine/Private/NetConnection.cpp
--- a/Engine/Source/Runtime/Engine/Private/NetConnection.cpp
+++ b/Engine/Source/Runtime/Engine/Private/NetConnection.cpp
@@ -955,7 +955,7 @@ void UNetConnection::Close(FNetResult&& CloseReason)
 		if (GNetCloseTimingDebug)
 		{
 			const double CurTime = FPlatformTime::Seconds();
-			const double TimeSinceRecv = LastReceiveRealtime != 0.0 ? (CurTime - LastReceiveRealtime) : -1.0;
+			const double TimeSinceRecv = LastReceiveRealtime != 0.0 && InTotalPackets != 0 ? (CurTime - LastReceiveRealtime) : -1.0;
 			const double TimeSinceSend = PreviousPacketSendTimeInS != 0.0 ? (CurTime - PreviousPacketSendTimeInS) : -1.0;
 			const double TimeSinceTickFlush = LastTime != 0.0 ? (CurTime - LastTime) : -1.0;
 			const double TimeSinceTickDispatch = CurTime - Driver->LastTickDispatchRealtime;
