Fix a voice over IP problem when the platform doesn't implement GetMicrophoneAmplitude

diff --git a/Engine/Plugins/Online/OnlineSubsystemUtils/Source/OnlineSubsystemUtils/Private/VoiceInterfaceImpl.cpp b/Engine/Plugins/Online/OnlineSubsystemUtils/Source/OnlineSubsystemUtils/Private/VoiceInterfaceImpl.cpp
--- a/Engine/Plugins/Online/OnlineSubsystemUtils/Source/OnlineSubsystemUtils/Private/VoiceInterfaceImpl.cpp
+++ b/Engine/Plugins/Online/OnlineSubsystemUtils/Source/OnlineSubsystemUtils/Private/VoiceInterfaceImpl.cpp
@@ -734,12 +734,14 @@ void FOnlineVoiceImpl::ProcessLocalVoicePackets()
 						uint32 Result = VoiceEngine->ReadLocalVoiceData(Index, BufferStart, &SpaceAvail, &SampleCount);
 
 						// Convert to Q15:
-						const float Amplitude = VoiceEngine->GetMicrophoneAmplitude(Index);
-						ensureAlways(Amplitude >= 0.0f && Amplitude <= 1.0f);
-
+						float Amplitude = VoiceEngine->GetMicrophoneAmplitude(Index);
+						if (!ensure(Amplitude >= 0.0f && Amplitude <= 1.0f))
+						{
+							// GetMicrophoneAmplitude returns -1 if not implemented which would mess up the MicrophoneAmplitude value so we set it to a sane value
+							Amplitude = 1.0f;
+						}
 						VoiceData.LocalPackets[Index].MicrophoneAmplitude = (int16)(Amplitude * 32767.0f);
 
-
 						if (Result == ONLINE_SUCCESS)
 						{
 							if (LocalTalkers[Index].bHasNetworkedVoice)
