Prevent ensure to be triggered during automated testing on Android
Valid until UIpConnection.Socket is deprecated

diff --git a/Engine/Plugins/Online/OnlineSubsystemUtils/Source/OnlineSubsystemUtils/Private/NetAddressResolution.cpp b/Engine/Plugins/Online/OnlineSubsystemUtils/Source/OnlineSubsystemUtils/Private/NetAddressResolution.cpp
--- a/Engine/Plugins/Online/OnlineSubsystemUtils/Source/OnlineSubsystemUtils/Private/NetAddressResolution.cpp
+++ b/Engine/Plugins/Online/OnlineSubsystemUtils/Source/OnlineSubsystemUtils/Private/NetAddressResolution.cpp
@@ -369,9 +369,12 @@ void FNetConnectionAddressResolution::CleanupResolutionSockets(ECleanupResolutio
 		if (EnumHasAnyFlags(CleanupFlags, ECleanupResolutionSocketsFlags::CleanAll))
 		{
 			// Remove when UIpConnection.Socket is deprecated
-#if !UE_BUILD_SHIPPING
-			ensureMsgf(DeprecatedSocket == nullptr,
-						TEXT("All codepaths passing 'CleanAll', must call 'UIpConnection::CleanupDeprecatedSocket' first."));
+#if !UE_BUILD_SHIPPING && !PLATFORM_ANDROID
+			if ( (!GIsBuildMachine) && (!IsRunningCommandlet()) ) // make sure we don't hit the ensure while on the Android testing... 
+			{
+				ensureMsgf(DeprecatedSocket == nullptr,
+					TEXT("All codepaths passing 'CleanAll', must call 'UIpConnection::CleanupDeprecatedSocket' first."));
+			}
 #endif
 
 			ResolutionSocket.Reset();
