[OnlinesubsystemUtils]
* Added ": int" to c enum in order to be able to forward declare it

diff --git a/Engine/Plugins/Online/OnlineSubsystemUtils/Source/OnlineSubsystemUtils/Public/PartyBeaconState.h b/Engine/Plugins/Online/OnlineSubsystemUtils/Source/OnlineSubsystemUtils/Public/PartyBeaconState.h
--- a/Engine/Plugins/Online/OnlineSubsystemUtils/Source/OnlineSubsystemUtils/Public/PartyBeaconState.h
+++ b/Engine/Plugins/Online/OnlineSubsystemUtils/Source/OnlineSubsystemUtils/Public/PartyBeaconState.h
@@ -16,7 +16,7 @@ ONLINESUBSYSTEMUTILS_API DECLARE_LOG_CATEGORY_EXTERN(LogPartyBeacon, Log, All);
 UENUM()
 namespace EPartyReservationResult
 {
-	enum Type
+	enum Type : int
 	{
 		/** Empty state. */
 		NoResult,
