Added detection for 'half-broken' NetConnection's, where only one side of the connection is receiving packets - and attempted recovery through recreating the network socket (which will use a new ephemeral port).

diff --git a/Engine/Plugins/Online/OnlineSubsystemSteam/Source/Private/SteamNetDriver.cpp b/Engine/Plugins/Online/OnlineSubsystemSteam/Source/Private/SteamNetDriver.cpp
--- a/Engine/Plugins/Online/OnlineSubsystemSteam/Source/Private/SteamNetDriver.cpp
+++ b/Engine/Plugins/Online/OnlineSubsystemSteam/Source/Private/SteamNetDriver.cpp
@@ -89,7 +89,12 @@ bool USteamNetDriver::InitConnect(FNetworkNotify* InNotify, const FURL& ConnectU
 		// If we are opening a Steam URL, create a Steam client socket
 		if (ConnectURL.Host.StartsWith(STEAM_URL_PREFIX))
 		{
-			SetSocketAndLocalAddress(SteamSockets->CreateSocket(FName(TEXT("SteamClientSocket")), TEXT("Unreal client (Steam)"), FNetworkProtocolTypes::Steam));
+			FUniqueSocket NewSocket = SteamSockets->CreateUniqueSocket(FName(TEXT("SteamClientSocket")), TEXT("Unreal client (Steam)"),
+																		FNetworkProtocolTypes::Steam);
+
+			TSharedPtr<FSocket> SharedSocket(NewSocket.Release(), FSocketDeleter(NewSocket.GetDeleter()));
+
+			SetSocketAndLocalAddress(SharedSocket);
 		}
 		else
 		{
@@ -106,7 +111,10 @@ bool USteamNetDriver::InitListen(FNetworkNotify* InNotify, FURL& ListenURL, bool
 	if (SteamSockets && !ListenURL.HasOption(TEXT("bIsLanMatch")) && !FParse::Param(FCommandLine::Get(), TEXT("forcepassthrough")))
 	{
 		FName SocketTypeName = IsRunningDedicatedServer() ? FName(TEXT("SteamServerSocket")) : FName(TEXT("SteamClientSocket"));
-		SetSocketAndLocalAddress(SteamSockets->CreateSocket(SocketTypeName, TEXT("Unreal server (Steam)"), FNetworkProtocolTypes::Steam));
+		FUniqueSocket NewSocket = SteamSockets->CreateUniqueSocket(SocketTypeName, TEXT("Unreal server (Steam)"), FNetworkProtocolTypes::Steam);
+		TSharedPtr<FSocket> SharedSocket(NewSocket.Release(), FSocketDeleter(NewSocket.GetDeleter()));
+
+		SetSocketAndLocalAddress(SharedSocket);
 	}
 	else
 	{
diff --git a/Engine/Plugins/Online/OnlineSubsystemUtils/Source/OnlineSubsystemUtils/Classes/IpConnection.h b/Engine/Plugins/Online/OnlineSubsystemUtils/Source/OnlineSubsystemUtils/Classes/IpConnection.h
--- a/Engine/Plugins/Online/OnlineSubsystemUtils/Source/OnlineSubsystemUtils/Classes/IpConnection.h
+++ b/Engine/Plugins/Online/OnlineSubsystemUtils/Source/OnlineSubsystemUtils/Classes/IpConnection.h
@@ -79,8 +79,13 @@ public:
 	FSocket* GetSocket() const
 	{
 		PRAGMA_DISABLE_DEPRECATION_WARNINGS
-		return Socket;
+		if (Socket != SocketPrivate.Get())
+		{
+			return Socket;
+		}
 		PRAGMA_ENABLE_DEPRECATION_WARNINGS
+
+		return SocketPrivate.Get();
 	}
 
 private:
@@ -90,14 +95,8 @@ private:
 	 */
 	struct FSocketSendResult
 	{
-		FSocketSendResult()
-			: BytesSent(0)
-			, Error(SE_NO_ERROR)
-		{
-		}
-
-		int32 BytesSent;
-		ESocketErrors Error;
+		int32 BytesSent = 0;
+		ESocketErrors Error = SE_NO_ERROR;
 	};
 
 	/** Critical section to protect SocketSendResults */
@@ -112,6 +111,18 @@ private:
 	 */
 	FGraphEventRef LastSendTask;
 
+	/** The socket used for communication (typically shared between NetConnection's, Net Address Resolution, and the NetDriver) */
+	TSharedPtr<FSocket> SocketPrivate;
+
+	/** List of previously active sockets for this connection, whose cleanup is deferred for multithreaded safety and to prevent remote ICMP errors */
+	TArray<TSharedPtr<FSocket>> DeferredCleanupSockets;
+
+	/** The time at which a socket was last queued for deferred cleanup */
+	double DeferredCleanupTimeCheck = 0.0;
+
+	/** The number of 'DeferredCleanupSockets' entries ready for cleanup (may be set by async send tasks) */
+	std::atomic<int32> DeferredCleanupReadyCount = 0;
+
 	/** Instead of disconnecting immediately on a socket error, wait for some time to see if we can recover. Specified in Seconds. */
 	UPROPERTY(Config)
 	float SocketErrorDisconnectDelay;
@@ -134,10 +145,15 @@ private:
 #endif
 
 private:
+	/** Sets the local socket pointer, and safely cleans up any references to old sockets */
+	void SetSocket_Local(const TSharedPtr<FSocket>& InSocket);
 
 	/** Cleanup for the deprecated 'Socket' value */
 	void CleanupDeprecatedSocket();
 
+	/** Safe/non-blocking cleanup of a shared socket, which may be in use by async sends, or may be at risk of triggering ICMP unreachable errors */
+	void SafeDeferredSocketCleanup(const TSharedPtr<FSocket>& InSocket);
+
 	/** Handles any SendTo errors on the game thread. */
 	void HandleSocketSendResult(const FSocketSendResult& Result, ISocketSubsystem* SocketSubsystem);
 
diff --git a/Engine/Plugins/Online/OnlineSubsystemUtils/Source/OnlineSubsystemUtils/Classes/IpNetDriver.h b/Engine/Plugins/Online/OnlineSubsystemUtils/Source/OnlineSubsystemUtils/Classes/IpNetDriver.h
--- a/Engine/Plugins/Online/OnlineSubsystemUtils/Source/OnlineSubsystemUtils/Classes/IpNetDriver.h
+++ b/Engine/Plugins/Online/OnlineSubsystemUtils/Source/OnlineSubsystemUtils/Classes/IpNetDriver.h
@@ -14,6 +14,7 @@
 #include "SocketTypes.h"
 #include "SocketSubsystem.h"
 #include "Templates/PimplPtr.h"
+#include "Containers/SpscQueue.h"
 #include "IpNetDriver.generated.h"
 
 
@@ -36,6 +37,28 @@ extern TSharedPtr<FInternetAddr> GCurrentDuplicateIP;
 #endif
 
 
+namespace UE::Net
+{
+	/** Results for UIpNetDriver::RecreateSocket */
+	enum class ERecreateSocketResult : uint8
+	{
+		NoAction,			// No action taken - this is either a replay, called on the server, or the NetDriver doesn't support socket recreation
+		NotReady,			// Net Address Resolution is in the middle of resolving, so no socket recreation can take place
+		AlreadyInProgress,	// Socket recreation is already in progress
+		BeganRecreate,		// Socket recreation was successfully kicked off
+		Error				// There was an error
+	};
+
+	const TCHAR* LexToString(ERecreateSocketResult Value);
+}
+
+namespace UE::Net::Private
+{
+	/** Callback triggered when setting a new socket has completed (which may not return immediately, due to multithreading) */
+	using FSetSocketComplete = TUniqueFunction<void()>;
+}
+
+
 UCLASS(transient, config=Engine)
 class ONLINESUBSYSTEMUTILS_API UIpNetDriver : public UNetDriver
 {
@@ -103,20 +126,19 @@ public:
 	virtual FSocket* GetSocket();
 
 	/**
-	 * Set the current NetDriver's socket to the given socket and takes ownership of it (the driver will handle destroying it).
-	 * This is typically done after resolution completes successfully. 
-	 * This will also set the LocalAddr for the netdriver automatically.
-	 *
-	 * @param NewSocket the socket pointer to set this netdriver's socket to
-	 */
+	* Set the current NetDriver's Socket/LocalAddr to the given socket (typically after Net Address Resolution).
+	* This will automatically pass the socket back through all NetConnection's, and trigger safe/deferred cleanup of the old socket.
+	*
+	* @param NewSocket	The socket pointer to set this NetDriver's socket to
+	*/
+	UE_DEPRECATED(5.1, "Use the version of SetSocketAndLocalAddress which takes a TSharedPtr")
 	void SetSocketAndLocalAddress(FSocket* NewSocket);
 
 	/**
-	* Set the current NetDriver's socket to the given socket.
-	* This is typically done after resolution completes successfully. 
-	* This will also set the LocalAddr for the netdriver automatically.
+	* Set the current NetDriver's Socket/LocalAddr to the given socket (typically after Net Address Resolution).
+	* This will automatically pass the socket back through all NetConnection's, and trigger safe/deferred cleanup of the old socket.
 	*
-	* @param NewSocket the socket pointer to set this netdriver's socket to
+	* @param NewSocket	The socket pointer to set this NetDriver's socket to
 	*/
 	void SetSocketAndLocalAddress(const TSharedPtr<FSocket>& SharedSocket);
 
@@ -129,6 +151,22 @@ public:
 	 */
 	virtual int GetClientPort();
 
+	/**
+	 * Creates a new Socket for the NetDriver/NetConnection's, based on the address/binding of the existing socket, with a new ephemeral port.
+	 * Used to attempt recovery for 'half-broken' connections, where (for whatever reason) only one side of the connection is receiving packets.
+	 *
+	 * Also usable to test the 'restart handshake' feature, using the 'net RecreateSocket' command.
+	 *
+	 * @return	Returns the action taken (e.g. successfully kicked off recreation, or whether this failed or is already in progress etc.)
+	 */
+	UE::Net::ERecreateSocketResult RecreateSocket();
+
+	/** Returns the last time socket recreation was kicked off or completed */
+	double GetLastRecreateSocketTime() const
+	{
+		return LastRecreateSocketTime;
+	}
+
 protected:
 	/**
 	 * Creates a socket set up for communication using the given protocol. This allows for explicit creation instead of inferring type for you.
@@ -153,6 +191,12 @@ protected:
 	virtual FUniqueSocket CreateAndBindSocket(TSharedRef<FInternetAddr> BindAddr, int32 Port, bool bReuseAddressAndPort, int32 DesiredRecvSize, int32 DesiredSendSize, FString& Error);
 	//~ End UIpNetDriver Interface.
 
+	/** Called by NetDriver subclasses, to specify whether or not they support recreating the active socket */
+	void SetSupportsRecreateSocket(bool bInSupportsRecreateSocket)
+	{
+		bSupportsRecreateSocket = bInSupportsRecreateSocket;
+	}
+
 public:
 	//~ Begin FExec Interface
 	virtual bool Exec( UWorld* InWorld, const TCHAR* Cmd, FOutputDevice& Ar=*GLog ) override;
@@ -163,6 +207,7 @@ public:
 
 	bool HandleSocketsCommand(const TCHAR* Cmd, FOutputDevice& Ar, UWorld* InWorld);
 	bool HandlePauseReceiveCommand(const TCHAR* Cmd, FOutputDevice& Ar, UWorld* InWorld);
+	bool HandleRecreateSocketCommand(const TCHAR* Cmd, FOutputDevice& Ar, UWorld* InWorld);
 
 #if !UE_BUILD_SHIPPING
 	/**
@@ -218,6 +263,39 @@ private:
 	void TickNewIPTracking(float DeltaTime);
 
 
+	/**
+	 * The current state of the NetDriver socket
+	 */
+	enum class ESocketState : uint8
+	{
+		None,
+		Resolving,		// Net Address Resolution is in the middle of resolving a working socket
+		Recreating,		// Socket Recreation is in the middle of creating a new socket
+		Ready,			// The current socket is fully ready to use
+		Error			// The current socket is in an error state
+	};
+
+
+	/** Returns the current state of the NetDriver socket */
+	ESocketState GetSocketState() const
+	{
+		return SocketState;
+	}
+
+	/**
+	 * Triggered when socket recreation completes (which may occur after a delay, due to e.g. updating the Receive Thread socket)
+	 */
+	void OnRecreateSocketComplete();
+
+	/**
+	 * Internal function for setting a socket and propagating it to all NetConnection's and the Receive Thread etc..
+	 *
+	 * @param InSocket			The new socket
+	 * @param InSetCallback		Callback triggered when the new socket is fully set
+	 */
+	void SetSocket_Internal(const TSharedPtr<FSocket>& InSocket, UE::Net::Private::FSetSocketComplete InSetCallback=nullptr);
+
+
 public:
 	// Callback for platform handling when networking is taking a long time in a single frame (by default over 1 second).
 	// It may get called multiple times in a single frame if additional processing after a previous alert exceeds the threshold again
@@ -283,6 +361,20 @@ private:
 
 		virtual uint32 Run() override;
 
+		/**
+		 * Execute commands in OwnerEventQueue on the Game Thread
+		 */
+		void PumpOwnerEventQueue();
+
+		/**
+		 * Sets the socket for the Receive Thread, and triggers a callback (if specified) on the Game Thread when this is complete.
+		 *
+		 * @param InGameThreadSocket		The socket specified by the Game Thread
+		 * @param InGameThreadSetCallback	The callback to trigger on the Game Thread when the socket is set
+		 */
+		void SetSocket(const TSharedPtr<FSocket>& InGameThreadSocket, UE::Net::Private::FSetSocketComplete InGameThreadSetCallback=nullptr);
+
+	public:
 		/** Thread-safe queue of received packets. The Run() function is the producer, UIpNetDriver::TickDispatch on the game thread is the consumer. */
 		TCircularQueue<FReceivedPacket> ReceiveQueue;
 
@@ -292,9 +384,34 @@ private:
 	private:
 		bool DispatchPacket(FReceivedPacket&& IncomingPacket, int32 NbBytesRead);
 
+		/**
+		 * Execute commands in CommandQueue on the Receive Thread
+		 */
+		void PumpCommandQueue();
+
+		/** Receive Thread implementation for 'SetSocket' */
+		void SetSocketImpl(const TSharedPtr<FSocket>& InGameThreadSocket, UE::Net::Private::FSetSocketComplete InGameThreadSetCallback=nullptr);
+
+		/**
+		 * Clears the Receive Thread socket, and resets the shared socket pointer on the Game Thread.
+		 */
+		void ClearSocket();
+
 	private:
+		/** Command queue from Game Thread to Receive Thread */
+		TSpscQueue<TUniqueFunction<void()>> CommandQueue;
+
+		/** Event queue from the Receive Thread to the Game Thread (may be generalized to 'Owning' thread in future) */
+		TSpscQueue<TUniqueFunction<void()>> OwnerEventQueue;
+
+		/** The NetDriver which owns the receive thread */
 		UIpNetDriver* OwningNetDriver;
+
+		/** The socket subsystem used by the receive thread */
 		ISocketSubsystem* SocketSubsystem;
+
+		/** Shared pointer for the socket, owned by the Game Thread */
+		TSharedPtr<FSocket> Socket;
 	};
 
 	/** Receive thread runnable object. */
@@ -309,6 +426,15 @@ private:
 	/** Underlying socket communication */
 	TSharedPtr<FSocket> SocketPrivate;
 
+	/** The state of the NetDriver socket (e.g. ready, resolving, recreating etc.) */
+	ESocketState SocketState = ESocketState::None;
+
+	/** Whether or not this NetDriver supports recreating the socket */
+	bool bSupportsRecreateSocket = true;
+
+	/** The last time socket recreation was kicked off or completed. */
+	double LastRecreateSocketTime = 0.0;
+
 
 	/** New IP aggregated logging entry */
 	struct FAggregatedIP
diff --git a/Engine/Plugins/Online/OnlineSubsystemUtils/Source/OnlineSubsystemUtils/Private/IpConnection.cpp b/Engine/Plugins/Online/OnlineSubsystemUtils/Source/OnlineSubsystemUtils/Private/IpConnection.cpp
--- a/Engine/Plugins/Online/OnlineSubsystemUtils/Source/OnlineSubsystemUtils/Private/IpConnection.cpp
+++ b/Engine/Plugins/Online/OnlineSubsystemUtils/Source/OnlineSubsystemUtils/Private/IpConnection.cpp
@@ -48,6 +48,7 @@ TAutoConsoleVariable<int32> CVarNetIpConnectionDisableResolution(
 	TEXT("If enabled, any future ip connections will not use resolution methods."),
 	ECVF_Default | ECVF_Cheat);
 
+
 namespace UE::Net::Private
 {
 #if !UE_BUILD_SHIPPING
@@ -64,9 +65,31 @@ namespace UE::Net::Private
 		TEXT("net.DebugInitialConnectLogFrequency"),
 		GNetDebugInitialConnectLogFrequency,
 		TEXT("The amount of time, in seconds, between initial connect debug logging."));
+
+	static int32 GNetBlockSend = 0;
+
+	static FAutoConsoleVariableRef CVarNetBlockSend(
+		TEXT("net.BlockSend"),
+		GNetBlockSend,
+		TEXT("When enabled, blocks packet sends on NetConnection's."));
 #endif
+
+	static int32 GNetRecreateSocketCooldown = 10;
+	static float GNetRecreateSocketTimeoutThreshold = 0.f;
+
+	static FAutoConsoleVariableRef CVarNetRecreateSocketCooldown(
+		TEXT("net.RecreateSocketCooldown"),
+		GNetRecreateSocketCooldown,
+		TEXT("The minimum amount of time, in seconds, between socket recreation attempts."));
+
+	static FAutoConsoleVariableRef CVarNetRecreateSocketTimeoutThreshold(
+		TEXT("net.RecreateSocketTimeoutThreshold"),
+		GNetRecreateSocketTimeoutThreshold,
+		TEXT("The amount of time, in seconds, without receiving a packet or alternatively without a send ack, before triggering socket recreation. ")
+		TEXT("(0.0 = off)"));
 }
 
+
 /**
  * UIpConnection
  */
@@ -186,9 +209,7 @@ void UIpConnection::Tick(float DeltaSeconds)
 
 	if (CheckResult == ECheckAddressResolutionResult::TryFirstAddress || CheckResult == ECheckAddressResolutionResult::TryNextAddress)
 	{
-		PRAGMA_DISABLE_DEPRECATION_WARNINGS
-		Socket = Resolver->GetResolutionSocket().Get();
-		PRAGMA_ENABLE_DEPRECATION_WARNINGS
+		SetSocket_Local(Resolver->GetResolutionSocket());
 
 		RemoteAddr = Resolver->GetRemoteAddr();
 
@@ -229,6 +250,16 @@ void UIpConnection::Tick(float DeltaSeconds)
 		CleanupDeprecatedSocket();
 	}
 
+	// Deferred cleanup of old sockets also waits until packets are received on the new socket (which happens after restart handshake process),
+	// in order to avoid getting kicked by the server, due to closing the old socket too early and triggering ICMP unreachable errors
+	if (DeferredCleanupSockets.Num() > 0 && ((LastReceiveRealtime - DeferredCleanupTimeCheck) >= 1.0) &&
+		DeferredCleanupSockets.Num() == DeferredCleanupReadyCount)
+	{
+		DeferredCleanupReadyCount = 0;
+
+		DeferredCleanupSockets.Reset();
+	}
+
 	Super::Tick(DeltaSeconds);
 
 #if !UE_BUILD_SHIPPING
@@ -248,7 +279,31 @@ void UIpConnection::Tick(float DeltaSeconds)
 			InitialConnectLastLogTime = CurTime;
 		}
 	}
-#endif	
+#endif
+
+	if (UIpNetDriver* IpDriver = Cast<UIpNetDriver>(Driver))
+	{
+		if (GNetRecreateSocketTimeoutThreshold != 0.f && GetConnectionState() != EConnectionState::USOCK_Closed && !IsReplay() &&
+			!IpDriver->IsServer())
+		{
+			const double ElapsedTime = Driver->GetElapsedTime();
+			const double LastRecreateSocketTime = IpDriver->GetLastRecreateSocketTime();
+			const double InitialConnectMultiplier = (ElapsedTime < Driver->InitialConnectTimeout && LastRecreateSocketTime == 0.0) ? 2.0 : 1.0;
+			const double RecreateCooldown = static_cast<double>(GNetRecreateSocketCooldown) * InitialConnectMultiplier;
+			const double CooldownCompareTimeDiff = ElapsedTime - FMath::Max(GetConnectTime(), LastRecreateSocketTime);
+			const double ThresholdCompareTimeDiff = ElapsedTime - FMath::Max(LastReceiveTime, GetLastRecvAckTime());
+
+			if (CooldownCompareTimeDiff > RecreateCooldown && ThresholdCompareTimeDiff > GNetRecreateSocketTimeoutThreshold)
+			{
+				UE::Net::ERecreateSocketResult Result = IpDriver->RecreateSocket();
+
+				UE_LOG(LogNet, Warning,
+						TEXT("Passed socket recreate threshold. LastReceiveTimeDiff: %f, LastRecvAckTimeDiff: %f, Threshold: %f, Result: %s"),
+						(ElapsedTime - LastReceiveTime), (ElapsedTime - GetLastRecvAckTime()), GNetRecreateSocketTimeoutThreshold,
+						ToCStr(UE::Net::LexToString(Result)));
+			}
+		}
+	}
 }
 
 void UIpConnection::CleanUp()
@@ -262,6 +317,34 @@ void UIpConnection::CleanUp()
 	// Sockets must be cleaned up after send tasks complete, as they can still be in use
 	CleanupDeprecatedSocket();
 	Resolver->CleanupResolutionSockets();
+
+	DeferredCleanupReadyCount = 0;
+
+	DeferredCleanupSockets.Empty();
+	SocketPrivate.Reset();
+}
+
+void UIpConnection::SetSocket_Local(const TSharedPtr<FSocket>& InSocket)
+{
+	TSharedPtr<FSocket> ResolutionSocket = Resolver->GetResolutionSocket();
+
+	if (ResolutionSocket.IsValid() && InSocket != ResolutionSocket)
+	{
+		SafeDeferredSocketCleanup(ResolutionSocket);
+		CleanupDeprecatedSocket();
+		Resolver->CleanupResolutionSockets();
+	}
+
+	if (SocketPrivate.IsValid() && InSocket != SocketPrivate && SocketPrivate != ResolutionSocket)
+	{
+		SafeDeferredSocketCleanup(SocketPrivate);
+	}
+
+	PRAGMA_DISABLE_DEPRECATION_WARNINGS
+	Socket = InSocket.Get();
+	PRAGMA_ENABLE_DEPRECATION_WARNINGS
+
+	SocketPrivate = InSocket;
 }
 
 void UIpConnection::CleanupDeprecatedSocket()
@@ -274,6 +357,35 @@ void UIpConnection::CleanupDeprecatedSocket()
 	PRAGMA_ENABLE_DEPRECATION_WARNINGS
 }
 
+// This needs to be removed when the main multithreading changes are added (or implemented in line with them)
+void UIpConnection::SafeDeferredSocketCleanup(const TSharedPtr<FSocket>& InSocket)
+{
+	if (!DeferredCleanupSockets.Contains(InSocket))
+	{
+		DeferredCleanupSockets.Add(InSocket);
+		DeferredCleanupTimeCheck = LastReceiveRealtime;
+
+		if (CVarNetIpConnectionUseSendTasks.GetValueOnGameThread() != 0 && LastSendTask.IsValid())
+		{
+			DECLARE_CYCLE_STAT(TEXT("IpConnection SocketCleanup task"), STAT_IpConnection_SocketCleanupTask, STATGROUP_TaskGraphTasks);
+
+			FGraphEventArray Prerequisites;
+
+			Prerequisites.Add(LastSendTask);
+
+			LastSendTask = FFunctionGraphTask::CreateAndDispatchWhenReady([this]()
+				{
+					this->DeferredCleanupReadyCount++;
+				},
+				GET_STATID(STAT_IpConnection_SocketCleanupTask), &Prerequisites);
+		}
+		else
+		{
+			DeferredCleanupReadyCount++;
+		}
+	}
+}
+
 void UIpConnection::HandleConnectionTimeout(const FString& ErrorStr)
 {
 	using namespace UE::Net::Private;
@@ -298,6 +410,8 @@ void UIpConnection::WaitForSendTasks()
 
 		SCOPE_CYCLE_COUNTER(STAT_IpConnection_WaitForSendTasks);
 		FTaskGraphInterface::Get().WaitUntilTaskCompletes(LastSendTask, ENamedThreads::GameThread);
+
+		DeferredCleanupSockets.Reset();
 	}
 }
 
@@ -335,6 +449,11 @@ void UIpConnection::LowLevelSend(void* Data, int32 CountBits, FOutPacketTraits&
 			RemoteAddr = OrigAddr;
 		}
 	};
+
+	if (GNetBlockSend)
+	{
+		return;
+	}
 #endif
 
 	if (Resolver->ShouldBlockSend())
diff --git a/Engine/Plugins/Online/OnlineSubsystemUtils/Source/OnlineSubsystemUtils/Private/IpNetDriver.cpp b/Engine/Plugins/Online/OnlineSubsystemUtils/Source/OnlineSubsystemUtils/Private/IpNetDriver.cpp
--- a/Engine/Plugins/Online/OnlineSubsystemUtils/Source/OnlineSubsystemUtils/Private/IpNetDriver.cpp
+++ b/Engine/Plugins/Online/OnlineSubsystemUtils/Source/OnlineSubsystemUtils/Private/IpNetDriver.cpp
@@ -173,6 +173,26 @@ namespace IPNetDriverInternal
 	}
 }
 
+namespace UE::Net
+{
+	const TCHAR* LexToString(ERecreateSocketResult Value)
+	{
+		switch (Value)
+		{
+		case ERecreateSocketResult::NoAction: return TEXT("NoAction");
+		case ERecreateSocketResult::NotReady: return TEXT("NotReady");
+		case ERecreateSocketResult::AlreadyInProgress: return TEXT("AlreadyInProgress");
+		case ERecreateSocketResult::BeganRecreate: return TEXT("BeganRecreate");
+		case ERecreateSocketResult::Error: return TEXT("Error");
+
+		default:
+			check(false);
+
+			return TEXT("Unknown");
+		}
+	}
+}
+
 /**
  * FPacketItrator
  *
@@ -715,6 +735,11 @@ private:
 	{
 		Connection->CleanupDeprecatedSocket();
 	}
+
+	static void SetSocket_Local(UIpConnection* Connection, const TSharedPtr<FSocket>& InSocket)
+	{
+		Connection->SetSocket_Local(InSocket);
+	}
 };
 
 
@@ -847,7 +872,8 @@ void UIpNetDriver::SetSocketAndLocalAddress(FSocket* NewSocket)
 
 void UIpNetDriver::SetSocketAndLocalAddress(const TSharedPtr<FSocket>& SharedSocket)
 {
-	SocketPrivate = SharedSocket;
+	// Must be called even if the current socket is already SharedSocket (for when Net Address Resolution resolves the current socket)
+	SetSocket_Internal(SharedSocket);
 
 	if (SocketPrivate.IsValid())
 	{
@@ -904,8 +930,6 @@ bool UIpNetDriver::InitBase( bool bInitAsClient, FNetworkNotify* InNotify, const
 		return false;
 	}
 
-
-	SetSocketAndLocalAddress(Resolver->GetFirstSocket());
 	
 	// If the cvar is set and the socket subsystem supports it, create the receive thread.
 	if (CVarNetIpNetDriverUseReceiveThread.GetValueOnAnyThread() != 0 && SocketSubsystem->IsSocketWaitSupported())
@@ -914,6 +938,8 @@ bool UIpNetDriver::InitBase( bool bInitAsClient, FNetworkNotify* InNotify, const
 		SocketReceiveThread.Reset(FRunnableThread::Create(SocketReceiveThreadRunnable.Get(), *FString::Printf(TEXT("IpNetDriver Receive Thread"), *NetDriverName.ToString())));
 	}
 
+	SetSocketAndLocalAddress(Resolver->GetFirstSocket());
+
 	bool bRecvMultiEnabled = CVarNetUseRecvMulti.GetValueOnAnyThread() != 0;
 	bool bRecvThreadEnabled = CVarNetIpNetDriverUseReceiveThread.GetValueOnAnyThread() != 0;
 
@@ -958,6 +984,8 @@ bool UIpNetDriver::InitBase( bool bInitAsClient, FNetworkNotify* InNotify, const
 
 bool UIpNetDriver::InitConnect( FNetworkNotify* InNotify, const FURL& ConnectURL, FString& Error )
 {
+	using namespace UE::Net::Private;
+
 	ISocketSubsystem* SocketSubsystem = GetSocketSubsystem();
 	if (SocketSubsystem == nullptr)
 	{
@@ -977,6 +1005,16 @@ bool UIpNetDriver::InitConnect( FNetworkNotify* InNotify, const FURL& ConnectURL
 	ServerConnection->InitLocalConnection(this, SocketPrivate.Get(), ConnectURL, USOCK_Pending);
 
 	Resolver->InitConnect(ServerConnection, SocketSubsystem, GetSocket(), ConnectURL);
+
+	UIpConnection* IpServerConnection = Cast<UIpConnection>(ServerConnection);
+
+	if (FNetConnectionAddressResolution* ConnResolver = FNetDriverAddressResolution::GetConnectionResolver(IpServerConnection))
+	{
+		if (ConnResolver->IsAddressResolutionEnabled() && !ConnResolver->IsAddressResolutionComplete())
+		{
+			SocketState = ESocketState::Resolving;
+		}
+	}
 	
 	UE_LOG(LogNet, Log, TEXT("Game client on port %i, rate %i"), ConnectURL.Port, ServerConnection->CurrentNetSpeed );
 	CreateInitialClientChannels();
@@ -1008,6 +1046,13 @@ void UIpNetDriver::TickDispatch(float DeltaTime)
 
 	Super::TickDispatch( DeltaTime );
 
+	const bool bUsingReceiveThread = SocketReceiveThreadRunnable.IsValid();
+
+	if (bUsingReceiveThread)
+	{
+		SocketReceiveThreadRunnable->PumpOwnerEventQueue();
+	}
+
 #if !UE_BUILD_SHIPPING
 	PauseReceiveEnd = (PauseReceiveEnd != 0.f && PauseReceiveEnd - (float)FPlatformTime::Seconds() > 0.f) ? PauseReceiveEnd : 0.f;
 
@@ -1613,6 +1658,25 @@ bool UIpNetDriver::HandlePauseReceiveCommand(const TCHAR* Cmd, FOutputDevice& Ar
 	return true;
 }
 
+// Doubles as an inbuilt restart handshake test
+bool UIpNetDriver::HandleRecreateSocketCommand(const TCHAR* Cmd, FOutputDevice& Ar, UWorld* InWorld)
+{
+	using namespace UE::Net;
+
+	ERecreateSocketResult Result = RecreateSocket();
+
+	if (Result == ERecreateSocketResult::BeganRecreate)
+	{
+		Ar.Logf(TEXT("Began socket recreation."));
+	}
+	else
+	{
+		Ar.Logf(TEXT("Failed to trigger socket recreation. Result: %s"), ToCStr(LexToString(Result)));
+	}
+
+	return true;
+}
+
 #if !UE_BUILD_SHIPPING
 void UIpNetDriver::TestSuddenPortChange(uint32 NumConnections)
 {
@@ -1656,6 +1720,10 @@ bool UIpNetDriver::Exec( UWorld* InWorld, const TCHAR* Cmd, FOutputDevice& Ar )
 	{
 		return HandlePauseReceiveCommand(Cmd, Ar, InWorld);
 	}
+	else if (FParse::Command(&Cmd, TEXT("RecreateSocket")))
+	{
+		return HandleRecreateSocketCommand(Cmd, Ar, InWorld);
+	}
 
 	return UNetDriver::Exec( InWorld, Cmd,Ar);
 }
@@ -1808,6 +1876,123 @@ void UIpNetDriver::TickNewIPTracking(float DeltaTime)
 	}
 }
 
+UE::Net::ERecreateSocketResult UIpNetDriver::RecreateSocket()
+{
+	using namespace UE::Net;
+	using namespace UE::Net::Private;
+
+	ERecreateSocketResult Result = ERecreateSocketResult::NoAction;
+
+	LastRecreateSocketTime = GetElapsedTime();
+
+	// Only clients can recreate the socket on-the-fly
+	if (bSupportsRecreateSocket && ServerConnection != nullptr && !ServerConnection->IsReplay())
+	{
+		const ESocketState CurState = GetSocketState();
+
+		if (CurState == ESocketState::Ready)
+		{
+			FString Error;
+			FUniqueSocket NewSocket = CreateAndBindSocket(LocalAddr.ToSharedRef(), GetClientPort(), false, ClientDesiredSocketReceiveBufferBytes,
+															ClientDesiredSocketSendBufferBytes, Error);
+
+			if (NewSocket.IsValid())
+			{
+				UE_LOG(LogNet, Log, TEXT("Recreated socket for bind address: %s"), ToCStr(LocalAddr->ToString(true)));
+
+				FSetSocketComplete SetCallback =
+					[ThisPtr = TWeakObjectPtr<UIpNetDriver>(this)]()
+					{
+						if (ThisPtr.IsValid())
+						{
+							ThisPtr->OnRecreateSocketComplete();
+						}
+					};
+
+				SocketState = ESocketState::Recreating;
+				Result = ERecreateSocketResult::BeganRecreate;
+
+				SetSocket_Internal(TSharedPtr<FSocket>(NewSocket.Release(), FSocketDeleter(NewSocket.GetDeleter())), MoveTemp(SetCallback));
+			}
+			else
+			{
+				UE_LOG(LogNet, Warning, TEXT("Could not recreate socket for bind address %s, got error %s"), ToCStr(LocalAddr->ToString(false)),
+						ToCStr(Error));
+
+				Result = ERecreateSocketResult::Error;
+			}
+		}
+		else if (CurState == ESocketState::Resolving)
+		{
+			Result = ERecreateSocketResult::NotReady;
+		}
+		else if (CurState == ESocketState::Recreating)
+		{
+			Result = ERecreateSocketResult::AlreadyInProgress;
+		}
+		else // if (CurState == ESocketState::Error)
+		{
+			Result = ERecreateSocketResult::Error;
+		}
+	}
+
+	return Result;
+}
+
+void UIpNetDriver::OnRecreateSocketComplete()
+{
+	SocketState = ESocketState::Ready;
+	LastRecreateSocketTime = GetElapsedTime();
+}
+
+void UIpNetDriver::SetSocket_Internal(const TSharedPtr<FSocket>& InSocket, UE::Net::Private::FSetSocketComplete InSetCallback/*=nullptr*/)
+{
+	if (SocketState != ESocketState::Recreating)
+	{
+		SocketState = ESocketState::Ready;
+	}
+
+	bool bHandledCallback = false;;
+
+	if (InSocket != SocketPrivate)
+	{
+		// Keep old socket around until we know we are done with it (don't want to trigger ICMP unreachable on the remote side by closing early)
+		TSharedPtr<FSocket> SavedSocket = MoveTemp(SocketPrivate);
+
+		SocketPrivate = InSocket;
+
+		const bool bUsingReceiveThread = SocketReceiveThreadRunnable.IsValid();
+
+		if (UIpConnection* IpServerConnection = GetServerConnection())
+		{
+			FIpConnectionHelper::SetSocket_Local(IpServerConnection, InSocket);
+		}
+		else
+		{
+			for (UNetConnection* CurConn : ClientConnections)
+			{
+				if (UIpConnection* IpConn = Cast<UIpConnection>(CurConn))
+				{
+					FIpConnectionHelper::SetSocket_Local(IpConn, InSocket);
+				}
+			}
+		}
+
+		if (bUsingReceiveThread)
+		{
+			SocketReceiveThreadRunnable->SetSocket(SocketPrivate, MoveTemp(InSetCallback));
+
+			bHandledCallback = true;
+		}
+	}
+
+	if (!bHandledCallback && InSetCallback)
+	{
+		InSetCallback();
+	}
+}
+
+
 UIpNetDriver::FReceiveThreadRunnable::FReceiveThreadRunnable(UIpNetDriver* InOwningNetDriver)
 	: ReceiveQueue(CVarNetIpNetDriverReceiveThreadQueueMaxPackets.GetValueOnAnyThread())
 	, bIsRunning(true)
@@ -1829,14 +2014,25 @@ uint32 UIpNetDriver::FReceiveThreadRunnable::Run()
 {
 	using namespace UE::Net::Private;
 
-	const FTimespan Timeout = FTimespan::FromMilliseconds(CVarNetIpNetDriverReceiveThreadPollTimeMS.GetValueOnAnyThread());
+	const int32 PollTimeMS = CVarNetIpNetDriverReceiveThreadPollTimeMS.GetValueOnAnyThread();
+	const FTimespan Timeout = FTimespan::FromMilliseconds(PollTimeMS);
 	const float SleepTimeForWaitableErrorsInSec = CVarRcvThreadSleepTimeForWaitableErrorsInSeconds.GetValueOnAnyThread();
 	const int32 ActionForLongRecvErrors = CVarRcvThreadShouldSleepForLongRecvErrors.GetValueOnAnyThread();
 
 	UE_LOG(LogNet, Log, TEXT("UIpNetDriver::FReceiveThreadRunnable::Run starting up."));
 
-	FSocket* CurSocket;
-	while (bIsRunning)
+	// Wait for the Socket to be set
+	while (bIsRunning && !Socket.IsValid())
+	{
+		// Process commands from the Game Thread
+		PumpCommandQueue();
+
+		const float NoSocketSetSleep = .03f;
+		FPlatformProcess::SleepNoStats(NoSocketSetSleep);
+	}
+
+
+	while (bIsRunning && Socket.IsValid())
 	{
 		// If we've encountered any errors during address resolution (this flag will not have the error state on it if resolution is disabled)
 		// Then stop running this thread. This stomps out any potential infinite loops caused by undefined behavior.
@@ -1844,21 +2040,19 @@ uint32 UIpNetDriver::FReceiveThreadRunnable::Run()
 
 		if (ServerConn != nullptr && FNetDriverAddressResolution::GetConnectionResolver(ServerConn)->HasAddressResolutionFailed())
 		{
-			break;
-		}
+			ClearSocket();
 
-		CurSocket = OwningNetDriver->GetSocket();
-		if (CurSocket == nullptr)
-		{
-			const float NoSocketSetSleep = .03f;
-			FPlatformProcess::SleepNoStats(NoSocketSetSleep);
-			continue;
+			break;
 		}
 
 		FReceivedPacket IncomingPacket;
 
 		bool bReceiveQueueFull = false;
-		if (CurSocket->Wait(ESocketWaitConditions::WaitForRead, Timeout))
+		const bool bWaitResult = Socket->Wait(ESocketWaitConditions::WaitForRead, Timeout);
+
+		PumpCommandQueue();
+
+		if (bWaitResult)
 		{
 			bool bOk = false;
 			int32 BytesRead = 0;
@@ -1869,7 +2063,7 @@ uint32 UIpNetDriver::FReceiveThreadRunnable::Run()
 
 			{
 				SCOPE_CYCLE_COUNTER(STAT_IpNetDriver_RecvFromSocket);
-				bOk = CurSocket->RecvFrom(IncomingPacket.PacketBytes.GetData(), IncomingPacket.PacketBytes.Num(), BytesRead, *IncomingPacket.FromAddress);
+				bOk = Socket->RecvFrom(IncomingPacket.PacketBytes.GetData(), IncomingPacket.PacketBytes.Num(), BytesRead, *IncomingPacket.FromAddress);
 			}
 
 			if (bOk)
@@ -1939,8 +2133,74 @@ uint32 UIpNetDriver::FReceiveThreadRunnable::Run()
 		}
 	}
 
+	// In case of a non-standard exit from the main loop, wait for a proper Game Thread triggered shutdown and cleanup
+	while (bIsRunning && Socket.IsValid())
+	{
+		PumpCommandQueue();
+		FPlatformProcess::SleepNoStats(PollTimeMS / 1000.f);
+	}
+
 	UE_LOG(LogNet, Log, TEXT("UIpNetDriver::FReceiveThreadRunnable::Run returning."));
 
 	return 0;
 }
 
+void UIpNetDriver::FReceiveThreadRunnable::PumpOwnerEventQueue()
+{
+	while (TOptional<TUniqueFunction<void()>> CurCommand = OwnerEventQueue.Dequeue())
+	{
+		(*CurCommand)();
+	}
+}
+
+void UIpNetDriver::FReceiveThreadRunnable::PumpCommandQueue()
+{
+	while (TOptional<TUniqueFunction<void()>> CurCommand = CommandQueue.Dequeue())
+	{
+		(*CurCommand)();
+	}
+}
+
+void UIpNetDriver::FReceiveThreadRunnable::SetSocket(const TSharedPtr<FSocket>& InGameThreadSocket,
+														UE::Net::Private::FSetSocketComplete InGameThreadSetCallback/*=nullptr*/)
+{
+	CommandQueue.Enqueue([this, GameThreadSocket = InGameThreadSocket, GameThreadSetCallback = MoveTemp(InGameThreadSetCallback)]() mutable
+	{
+		this->SetSocketImpl(GameThreadSocket, MoveTemp(GameThreadSetCallback));
+	});
+}
+
+void UIpNetDriver::FReceiveThreadRunnable::SetSocketImpl(const TSharedPtr<FSocket>& InGameThreadSocket,
+														UE::Net::Private::FSetSocketComplete InGameThreadSetCallback/*=nullptr*/)
+{
+	ClearSocket();
+
+	Socket = InGameThreadSocket;
+
+	if (InGameThreadSetCallback)
+	{
+		OwnerEventQueue.Enqueue([GameThreadSetCallback = MoveTemp(InGameThreadSetCallback)]() mutable
+		{
+			GameThreadSetCallback();
+		});
+	}
+}
+
+void UIpNetDriver::FReceiveThreadRunnable::ClearSocket()
+{
+	if (Socket.IsValid())
+	{
+		TSharedPtr<FSocket> GameThreadSocket = MoveTemp(Socket);
+
+		Socket.Reset();
+
+		// Clear the socket shared pointer from the Game Thread
+		if (GameThreadSocket.IsValid())
+		{
+			OwnerEventQueue.Enqueue([GameThreadSocket = MoveTemp(GameThreadSocket)]() mutable
+			{
+				GameThreadSocket.Reset();
+			});
+		}
+	}
+}
diff --git a/Engine/Plugins/Online/SocketSubsystemEOS/Source/SocketSubsystemEOS/Private/NetDriverEOSBase.cpp b/Engine/Plugins/Online/SocketSubsystemEOS/Source/SocketSubsystemEOS/Private/NetDriverEOSBase.cpp
--- a/Engine/Plugins/Online/SocketSubsystemEOS/Source/SocketSubsystemEOS/Private/NetDriverEOSBase.cpp
+++ b/Engine/Plugins/Online/SocketSubsystemEOS/Source/SocketSubsystemEOS/Private/NetDriverEOSBase.cpp
@@ -60,8 +60,12 @@ bool UNetDriverEOSBase::InitBase(bool bInitAsClient, FNetworkNotify* InNotify, c
 		return false;
 	}
 
-	// Create our socket
-	SetSocketAndLocalAddress(SocketSubsystem->CreateSocket(NAME_DGram, TEXT("UE4"), NAME_None));
+
+	FUniqueSocket NewSocket = SocketSubsystem->CreateUniqueSocket(NAME_DGram, TEXT("UE4"), NAME_None);
+	TSharedPtr<FSocket> SharedSocket(NewSocket.Release(), FSocketDeleter(NewSocket.GetDeleter()));
+
+	SetSocketAndLocalAddress(SharedSocket);
+
 	if (GetSocket() == nullptr)
 	{
 		UE_LOG(LogTemp, Warning, TEXT("Could not create socket"));
