- Disable/extend beacon client timeouts in -notimeouts / unoptimized builds

diff --git a/Engine/Plugins/Online/OnlineSubsystemUtils/Source/OnlineSubsystemUtils/Private/OnlineBeaconClient.cpp b/Engine/Plugins/Online/OnlineSubsystemUtils/Source/OnlineSubsystemUtils/Private/OnlineBeaconClient.cpp
--- a/Engine/Plugins/Online/OnlineSubsystemUtils/Source/OnlineSubsystemUtils/Private/OnlineBeaconClient.cpp
+++ b/Engine/Plugins/Online/OnlineSubsystemUtils/Source/OnlineSubsystemUtils/Private/OnlineBeaconClient.cpp
@@ -18,8 +18,6 @@
 
 #include UE_INLINE_GENERATED_CPP_BY_NAME(OnlineBeaconClient)
 
-#define BEACON_RPC_TIMEOUT 15.0f
-
 /** For backwards compatibility with newer engine encryption API */
 #ifndef NETCONNECTION_HAS_SETENCRYPTIONKEY
 	#define NETCONNECTION_HAS_SETENCRYPTIONKEY 0
@@ -483,8 +481,11 @@ void AOnlineBeaconClient::NotifyControlMessage(UNetConnection* Connection, uint8
 						// Server will send ClientOnConnected() when it gets this control message
 
 						// Fail safe for connection to server but no client connection RPC
-						FTimerDelegate TimerDelegate = FTimerDelegate::CreateUObject(this, &AOnlineBeaconClient::OnFailure);
-						GetWorldTimerManager().SetTimer(TimerHandle_OnFailure, TimerDelegate, BEACON_RPC_TIMEOUT, false);
+						if (!Connection->Driver->bNoTimeouts)
+						{
+							FTimerDelegate TimerDelegate = FTimerDelegate::CreateUObject(this, &AOnlineBeaconClient::OnFailure);
+							GetWorldTimerManager().SetTimer(TimerHandle_OnFailure, TimerDelegate, Connection->GetTimeoutValue(), false);
+						}
 					}
 					else
 					{
