Add support for ID tokens with beacon authentication.

diff --git a/Engine/Plugins/Online/OnlineSubsystemUtils/Source/OnlineSubsystemUtils/Private/OnlineBeaconClient.cpp b/Engine/Plugins/Online/OnlineSubsystemUtils/Source/OnlineSubsystemUtils/Private/OnlineBeaconClient.cpp
--- a/Engine/Plugins/Online/OnlineSubsystemUtils/Source/OnlineSubsystemUtils/Private/OnlineBeaconClient.cpp
+++ b/Engine/Plugins/Online/OnlineSubsystemUtils/Source/OnlineSubsystemUtils/Private/OnlineBeaconClient.cpp
@@ -409,19 +409,12 @@ void AOnlineBeaconClient::NotifyControlMessage(UNetConnection* Connection, uint8
 			{
 				// build a URL
 				FURL URL(nullptr, TEXT(""), TRAVEL_Absolute);
-				FString URLString;
-
-				// Append authentication token to URL options
-				IOnlineIdentityPtr IdentityPtr = Online::GetIdentityInterface(GetWorld());
-				if (IdentityPtr.IsValid())
+				FString AuthTicket = GetAuthTicket(*Connection->PlayerId);
+				if (!AuthTicket.IsEmpty())
 				{
-					TSharedPtr<FUserOnlineAccount> UserAcct = IdentityPtr->GetUserAccount(*Connection->PlayerId);
-					if (UserAcct.IsValid())
-					{
-						URL.AddOption(*FString::Printf(TEXT("AuthTicket=%s"), *UserAcct->GetAccessToken()));
-					}
+					URL.AddOption(*FString::Printf(TEXT("AuthTicket=%s"), *AuthTicket));
 				}
-				URLString = URL.ToString();
+				FString URLString = URL.ToString();
 
 				// compute the player's online platform name
 				FName OnlinePlatformName = NAME_None;
@@ -556,6 +549,23 @@ void AOnlineBeaconClient::NotifyControlMessage(UNetConnection* Connection, uint8
 	}	
 }
 
+FString AOnlineBeaconClient::GetAuthTicket(const FUniqueNetIdRepl& PlayerId)
+{
+	FString AuthTicket;
+
+	IOnlineIdentityPtr IdentityPtr = Online::GetIdentityInterface(GetWorld());
+	if (IdentityPtr.IsValid())
+	{
+		TSharedPtr<FUserOnlineAccount> UserAcct = IdentityPtr->GetUserAccount(*PlayerId);
+		if (UserAcct.IsValid())
+		{
+			AuthTicket = *UserAcct->GetAccessToken();
+		}
+	}
+
+	return AuthTicket;
+}
+
 void AOnlineBeaconClient::FinalizeEncryptedConnection(const FEncryptionKeyResponse& Response, TWeakObjectPtr<UNetConnection> WeakConnection)
 {
 	UNetConnection* Connection = WeakConnection.Get();
diff --git a/Engine/Plugins/Online/OnlineSubsystemUtils/Source/OnlineSubsystemUtils/Public/OnlineBeaconClient.h b/Engine/Plugins/Online/OnlineSubsystemUtils/Source/OnlineSubsystemUtils/Public/OnlineBeaconClient.h
--- a/Engine/Plugins/Online/OnlineSubsystemUtils/Source/OnlineSubsystemUtils/Public/OnlineBeaconClient.h
+++ b/Engine/Plugins/Online/OnlineSubsystemUtils/Source/OnlineSubsystemUtils/Public/OnlineBeaconClient.h
@@ -131,6 +131,15 @@ class ONLINESUBSYSTEMUTILS_API AOnlineBeaconClient : public AOnlineBeacon
 	 */
 	const FUniqueNetIdRepl& GetUniqueId() const;
 
+	/**
+	* Get the authentication ticket to send when the server requires auth.
+	* 
+	* @param PlayerId unique id of the user on this connection
+	* 
+	* @return AuthTicket value to send to server
+	*/
+	virtual FString GetAuthTicket(const FUniqueNetIdRepl& PlayerId);
+
 	/**
 	 * Get the connection state
 	 * Higher level than the net connection because of the handshaking of the actors
