Check if PlayerId is valid.

diff --git a/Engine/Plugins/Online/OnlineSubsystemUtils/Source/OnlineSubsystemUtils/Private/OnlineBeaconClient.cpp b/Engine/Plugins/Online/OnlineSubsystemUtils/Source/OnlineSubsystemUtils/Private/OnlineBeaconClient.cpp
--- a/Engine/Plugins/Online/OnlineSubsystemUtils/Source/OnlineSubsystemUtils/Private/OnlineBeaconClient.cpp
+++ b/Engine/Plugins/Online/OnlineSubsystemUtils/Source/OnlineSubsystemUtils/Private/OnlineBeaconClient.cpp
@@ -554,7 +554,7 @@ FString AOnlineBeaconClient::GetAuthTicket(const FUniqueNetIdRepl& PlayerId)
 	FString AuthTicket;
 
 	IOnlineIdentityPtr IdentityPtr = Online::GetIdentityInterface(GetWorld());
-	if (IdentityPtr.IsValid())
+	if (IdentityPtr.IsValid() && ensure(PlayerId.IsValid()))
 	{
 		TSharedPtr<FUserOnlineAccount> UserAcct = IdentityPtr->GetUserAccount(*PlayerId);
 		if (UserAcct.IsValid())
