Fix beacon authentication hang which occurs when multiple beacons for the same player are handshaking at the same time.

diff --git a/Engine/Plugins/Online/OnlineSubsystemUtils/Source/OnlineSubsystemUtils/Private/OnlineBeaconHost.cpp b/Engine/Plugins/Online/OnlineSubsystemUtils/Source/OnlineSubsystemUtils/Private/OnlineBeaconHost.cpp
--- a/Engine/Plugins/Online/OnlineSubsystemUtils/Source/OnlineSubsystemUtils/Private/OnlineBeaconHost.cpp
+++ b/Engine/Plugins/Online/OnlineSubsystemUtils/Source/OnlineSubsystemUtils/Private/OnlineBeaconHost.cpp
@@ -294,7 +294,27 @@ bool AOnlineBeaconHost::HandleControlMessage(UNetConnection* Connection, uint8 M
 			Connection->SetPlayerOnlinePlatformName(FName(*OnlinePlatformName));
 
 			// Try to kick off verification for this player.
-			if (!StartVerifyAuthentication(*UniqueIdRepl, UGameplayStatics::ParseOption(OptionsURL, TEXT("AuthTicket"))))
+			const FString AuthTicket = UGameplayStatics::ParseOption(OptionsURL, TEXT("AuthTicket"));
+
+			// Try to start deprecated auth method.
+			PRAGMA_DISABLE_DEPRECATION_WARNINGS
+			const bool bStartedAuthDeprecated = StartVerifyAuthentication(*UniqueIdRepl, AuthTicket);
+			PRAGMA_ENABLE_DEPRECATION_WARNINGS
+
+			// Don't start new auth method if deprecated auth is active.
+			bool bStartedNewAuth = false;
+			if (!bStartedAuthDeprecated)
+			{
+				// Create completion delegate.
+				FOnAuthenticationVerificationCompleteDelegate OnAuthComplete = FOnAuthenticationVerificationCompleteDelegate::CreateWeakLambda(this, [this, WeakConnection = TWeakObjectPtr<UNetConnection>(Connection)](const FOnlineError& OnlineError)
+				{
+					OnAuthenticationVerificationComplete(WeakConnection.Get(), OnlineError);
+				});
+
+				bStartedNewAuth = StartVerifyAuthentication(*UniqueIdRepl, AuthTicket, OnAuthComplete);
+			}
+
+			if (!bStartedAuthDeprecated && !bStartedNewAuth)
 			{
 				static const FText ErrorTxt = NSLOCTEXT("NetworkErrors", "BeaconLoginInvalidAuthHandlerError", "Login Failure. Unable to process authentication.");
 				SendFailurePacket(Connection, ENetCloseResult::BeaconLoginInvalidAuthHandlerError, ErrorTxt);
@@ -579,6 +599,19 @@ void AOnlineBeaconHost::OnAuthenticationVerificationComplete(const class FUnique
 	UNetConnection* Connection = nullptr;
 	FConnectionState* ConnState = nullptr;
 	if (GetConnectionDataForUniqueNetId(PlayerId, Connection, ConnState))
+	{
+		OnAuthenticationVerificationComplete(Connection, Error);
+	}
+}
+
+bool AOnlineBeaconHost::StartVerifyAuthentication(const FUniqueNetId& PlayerId, const FString& AuthenticationToken, const FOnAuthenticationVerificationCompleteDelegate& OnComplete)
+{
+	return false;
+}
+
+void AOnlineBeaconHost::OnAuthenticationVerificationComplete(UNetConnection* Connection, const FOnlineError& Error)
+{
+	if (FConnectionState* ConnState = ConnectionState.Find(Connection))
 	{
 		if (ConnState->bHasSentLogin && !ConnState->bHasSentWelcome)
 		{
@@ -589,7 +622,7 @@ void AOnlineBeaconHost::OnAuthenticationVerificationComplete(const class FUnique
 				ConnState->bHasAuthenticated = true;
 
 				// send the welcome packet
-				UE_LOG(LogBeacon, Verbose, TEXT("%s: User authenticated: %s"), *GetDebugName(Connection), *PlayerId.ToString());
+				UE_LOG(LogBeacon, Verbose, TEXT("%s: OnAuthenticationVerificationComplete: Successfully authenticated."), *GetDebugName(Connection));
 				ConnState->bHasSentWelcome = true;
 				FNetControlMessage<NMT_BeaconWelcome>::Send(Connection);
 				Connection->FlushNet();
@@ -597,11 +630,19 @@ void AOnlineBeaconHost::OnAuthenticationVerificationComplete(const class FUnique
 			else
 			{
 				// Auth failure.
-				UE_LOG(LogBeacon, Log, TEXT("%s: Failed to authenticate user: %s error: %s"), *GetDebugName(Connection), *PlayerId.ToString(), *Error.ToLogString());
+				UE_LOG(LogBeacon, Log, TEXT("%s: OnAuthenticationVerificationComplete: Failed to authenticate error: %s"), *GetDebugName(Connection), *Error.ToLogString());
 				SendFailurePacket(Connection, ENetCloseResult::BeaconAuthenticationFailure, Error_Authentication);
 				CloseHandshakeConnection(Connection);
 			}
 		}
+		else
+		{
+			UE_LOG(LogBeacon, Warning, TEXT("%s: OnAuthenticationVerificationComplete: Invalid beacon connection state for auth challenge result. error: %s"), *GetDebugName(Connection), *Error.ToLogString());
+		}
+	}
+	else
+	{
+		UE_LOG(LogBeacon, Warning, TEXT("%s: OnAuthenticationVerificationComplete: Failed to find connection state. error: %s"), *GetDebugName(Connection), *Error.ToLogString());
 	}
 }
 
diff --git a/Engine/Plugins/Online/OnlineSubsystemUtils/Source/OnlineSubsystemUtils/Public/OnlineBeaconHost.h b/Engine/Plugins/Online/OnlineSubsystemUtils/Source/OnlineSubsystemUtils/Public/OnlineBeaconHost.h
--- a/Engine/Plugins/Online/OnlineSubsystemUtils/Source/OnlineSubsystemUtils/Public/OnlineBeaconHost.h
+++ b/Engine/Plugins/Online/OnlineSubsystemUtils/Source/OnlineSubsystemUtils/Public/OnlineBeaconHost.h
@@ -122,6 +122,7 @@ protected:
 	 * @param PlayerId net id of player to authenticate.
 	 * @param AuthenticationToken token to use for verification.
 	 */
+	 UE_DEPRECATED(5.2, "This version of the StartVerifyAuthentication is deprecated. Please use the new StartVerifyAuthentication method instead.")
 	virtual bool StartVerifyAuthentication(const FUniqueNetId& PlayerId, const FString& AuthenticationToken);
 
 	/**
@@ -130,8 +131,35 @@ protected:
 	 * @param PlayerId net id of player to authenticate.
 	 * @param Error result of the operation.
 	 */
+	UE_DEPRECATED(5.2, "This version of the OnAuthenticationVerificationComplete is deprecated. Please use the new OnAuthenticationVerificationComplete method instead.")
 	void OnAuthenticationVerificationComplete(const class FUniqueNetId& PlayerId, const FOnlineError& Error);
 
+	/**
+	 * Delegate executed when user authentication has completed.
+	 *
+	 * @param OnlineError the result of the operation
+	 */
+	DECLARE_DELEGATE_OneParam(FOnAuthenticationVerificationCompleteDelegate, const FOnlineError& /*OnlineError*/);
+
+	/**
+	 * Start verifying an authentication token for a connection.
+	 * OnAuthenticationVerificationComplete must be called to complete authentication verification.
+	 *
+	 * @param NetConnection network connection associated with the authentication challenge.
+	 * @param PlayerId net id of player to authenticate.
+	 * @param AuthenticationToken token to use for verification.
+	 */
+	virtual bool StartVerifyAuthentication(const FUniqueNetId& PlayerId, const FString& AuthenticationToken, const FOnAuthenticationVerificationCompleteDelegate& OnComplete);
+
+private:
+	/**
+	 * Event which must be signaled to complete an authentication verification request.
+	 *
+	 * @param NetConnection network connection associated with the authentication challenge.
+	 * @param Error result of the operation.
+	 */
+	void OnAuthenticationVerificationComplete(UNetConnection* Connection, const FOnlineError& Error);
+
 protected:
 	FString GetDebugName(UNetConnection* Connection = nullptr);
 
