Add timeout logging and RPC delay CVar for beacon.

diff --git a/Engine/Plugins/Online/OnlineSubsystemUtils/Source/OnlineSubsystemUtils/Private/OnlineBeaconClient.cpp b/Engine/Plugins/Online/OnlineSubsystemUtils/Source/OnlineSubsystemUtils/Private/OnlineBeaconClient.cpp
--- a/Engine/Plugins/Online/OnlineSubsystemUtils/Source/OnlineSubsystemUtils/Private/OnlineBeaconClient.cpp
+++ b/Engine/Plugins/Online/OnlineSubsystemUtils/Source/OnlineSubsystemUtils/Private/OnlineBeaconClient.cpp
@@ -311,6 +311,8 @@ void AOnlineBeaconClient::OnFailure()
 
 void AOnlineBeaconClient::ClientOnConnected_Implementation()
 {
+	UE_LOG(LogBeacon, Verbose, TEXT("%s: Received first RPC from server"), *GetName());
+
 	SetConnectionState(EBeaconConnectionState::Open);
 	BeaconConnection->SetConnectionState(USOCK_Open);
 
@@ -476,8 +478,10 @@ void AOnlineBeaconClient::NotifyControlMessage(UNetConnection* Connection, uint8
 						// Fail safe for connection to server but no client connection RPC
 						if (!Connection->Driver->bNoTimeouts)
 						{
+							float Timeout = Connection->GetTimeoutValue();
+							UE_LOG(LogBeacon, Verbose, TEXT("%s: Waiting for first RPC from server, timeout: %.3f"), *GetName(), Timeout);
 							FTimerDelegate TimerDelegate = FTimerDelegate::CreateUObject(this, &AOnlineBeaconClient::OnFailure);
-							GetWorldTimerManager().SetTimer(TimerHandle_OnFailure, TimerDelegate, Connection->GetTimeoutValue(), false);
+							GetWorldTimerManager().SetTimer(TimerHandle_OnFailure, TimerDelegate, Timeout, false);
 						}
 					}
 					else
diff --git a/Engine/Plugins/Online/OnlineSubsystemUtils/Source/OnlineSubsystemUtils/Private/OnlineBeaconHost.cpp b/Engine/Plugins/Online/OnlineSubsystemUtils/Source/OnlineSubsystemUtils/Private/OnlineBeaconHost.cpp
--- a/Engine/Plugins/Online/OnlineSubsystemUtils/Source/OnlineSubsystemUtils/Private/OnlineBeaconHost.cpp
+++ b/Engine/Plugins/Online/OnlineSubsystemUtils/Source/OnlineSubsystemUtils/Private/OnlineBeaconHost.cpp
@@ -15,6 +15,18 @@
 
 #include UE_INLINE_GENERATED_CPP_BY_NAME(OnlineBeaconHost)
 
+#if !UE_BUILD_SHIPPING
+namespace BeaconConsoleVariables
+{
+	TAutoConsoleVariable<float> CVarDelayFinishHandshake(
+		TEXT("beacon.DelayFinishHandshake"),
+		0.0f,
+		TEXT("Delay time before finishing handshake by calling client RPC\n")
+		TEXT("Time in seconds. A value of 0 means no delay, and a negative value means never call it."),
+		ECVF_Default);
+}
+#endif
+
 const FText Error_UnableToParsePacket = NSLOCTEXT("NetworkErrors", "UnableToParsePacket", "Unable to parse expected packet structure: {0}.");
 const FText Error_ControlFlow = NSLOCTEXT("NetworkErrors", "ControlFlowError", "Control flow error: {0}.");
 const FText Error_Authentication = NSLOCTEXT("NetworkErrors", "AuthenticationFailure", "Failed to verify user authentication.");
@@ -148,7 +160,7 @@ bool AOnlineBeaconHost::HandleControlMessage(UNetConnection* Connection, uint8 M
 		}
 
 		// make one
-		ConnState = &ConnectionState.Add(Connection, FConnectionState());
+		ConnState = &ConnectionState.Add(Connection, FConnectionState(*this));
 		Connection->OwningActor = this; // make sure we get OnNetCleanup if this dies
 	}
 	check(ConnState != nullptr);
@@ -455,35 +467,27 @@ bool AOnlineBeaconHost::HandleControlMessage(UNetConnection* Connection, uint8 M
 				return false;
 			}
 
-			AOnlineBeaconClient* ClientActor = GetClientActor(Connection);
-			if (ClientActor == nullptr || BeaconType != ClientActor->GetBeaconType())
+#if UE_BUILD_SHIPPING
+			const float Delay = 0.0f;
+#else
+			const float Delay = BeaconConsoleVariables::CVarDelayFinishHandshake.GetValueOnGameThread();
+#endif
+			if (Delay == 0.0f)
 			{
-				static const FText ErrorTxt = NSLOCTEXT("NetworkErrors", "BeaconSpawnNetGUIDAckError2", "Join failure, no actor at NetGUIDAck.");
-				SendFailurePacket(Connection, ENetCloseResult::BeaconSpawnNetGUIDAckNoActor, ErrorTxt);
-				return false;
+				FinishHandshake(Connection, MoveTemp(BeaconType));
 			}
-
-			FOnBeaconConnected* OnBeaconConnectedDelegate = OnBeaconConnectedMapping.Find(BeaconType);
-			if (OnBeaconConnectedDelegate == nullptr)
+			else if (Delay > 0.0f)
 			{
-				static const FText ErrorTxt = NSLOCTEXT("NetworkErrors", "BeaconSpawnNetGUIDAckError1", "Join failure, no host object at NetGUIDAck.");
-				SendFailurePacket(Connection, ENetCloseResult::BeaconSpawnNetGUIDAckNoHost, ErrorTxt);
-				return false;
+				UE_LOG(LogBeacon, Verbose, TEXT("%s: Delay for handshake completion: %.3f"), *GetDebugName(Connection), Delay);
+				FTimerDelegate TimerDelegate = FTimerDelegate::CreateUObject(this, &AOnlineBeaconHost::FinishHandshake, Connection, MoveTemp(BeaconType));
+				GetWorldTimerManager().SetTimer(ConnState->FinishHandshakeTimerHandle, TimerDelegate, Delay, false);
+			}
+			else
+			{
+				UE_LOG(LogBeacon, Verbose, TEXT("%s: Handshake will never complete, client will either timeout or hang if timeouts are disabled."), *GetDebugName(Connection));
 			}
-
-			// once we change the owning actor we can't clean up state anymore, but at this point we don't need to
-			ConnectionState.Remove(Connection);
-			Connection->OwningActor = ClientActor;
-
-			ClientActor->SetReplicates(true);
-			ClientActor->SetAutonomousProxy(true);
-			ClientActor->SetConnectionState(EBeaconConnectionState::Open);
-			// Send an RPC to the client to open the actor channel and guarantee RPCs will work
-			ClientActor->ClientOnConnected();
-			UE_LOG(LogBeacon, Log, TEXT("%s: Handshake complete."), *GetDebugName(Connection));
-
-			OnBeaconConnectedDelegate->ExecuteIfBound(ClientActor, Connection);
 		}
+
 		break;
 
 	default:
@@ -498,6 +502,45 @@ bool AOnlineBeaconHost::HandleControlMessage(UNetConnection* Connection, uint8 M
 	return true;
 }
 
+void AOnlineBeaconHost::FinishHandshake(UNetConnection* Connection, FString BeaconType)
+{
+	AOnlineBeaconClient* ClientActor = GetClientActor(Connection);
+	if (ClientActor == nullptr || BeaconType != ClientActor->GetBeaconType())
+	{
+		static const FText ErrorTxt = NSLOCTEXT("NetworkErrors", "BeaconSpawnNetGUIDAckError2", "Join failure, no actor at NetGUIDAck.");
+		SendFailurePacket(Connection, ENetCloseResult::BeaconSpawnNetGUIDAckNoActor, ErrorTxt);
+		CloseHandshakeConnection(Connection);
+		return;
+	}
+
+	FOnBeaconConnected* OnBeaconConnectedDelegate = OnBeaconConnectedMapping.Find(BeaconType);
+	if (OnBeaconConnectedDelegate == nullptr)
+	{
+		static const FText ErrorTxt = NSLOCTEXT("NetworkErrors", "BeaconSpawnNetGUIDAckError1", "Join failure, no host object at NetGUIDAck.");
+		SendFailurePacket(Connection, ENetCloseResult::BeaconSpawnNetGUIDAckNoHost, ErrorTxt);
+		CloseHandshakeConnection(Connection);
+		return;
+	}
+
+	// once we change the owning actor we can't clean up state anymore, but at this point we don't need to
+	ConnectionState.Remove(Connection);
+	Connection->OwningActor = ClientActor;
+
+	ClientActor->SetReplicates(true);
+	ClientActor->SetAutonomousProxy(true);
+	ClientActor->SetConnectionState(EBeaconConnectionState::Open);
+
+	// Send an RPC to the client to open the actor channel and guarantee RPCs will work
+	ClientActor->ClientOnConnected();
+	UE_LOG(LogBeacon, Log, TEXT("%s: Handshake complete."), *GetDebugName(Connection));
+	OnBeaconConnectedDelegate->ExecuteIfBound(ClientActor, Connection);
+}
+
+AOnlineBeaconHost::FConnectionState::~FConnectionState()
+{
+	BeaconHost.GetWorldTimerManager().ClearTimer(FinishHandshakeTimerHandle);
+}
+
 void AOnlineBeaconHost::SendFailurePacket(UNetConnection* Connection, FNetCloseResult&& CloseReason, const FText& ErrorText)
 {
 	if (Connection != nullptr)
diff --git a/Engine/Plugins/Online/OnlineSubsystemUtils/Source/OnlineSubsystemUtils/Public/OnlineBeaconHost.h b/Engine/Plugins/Online/OnlineSubsystemUtils/Source/OnlineSubsystemUtils/Public/OnlineBeaconHost.h
--- a/Engine/Plugins/Online/OnlineSubsystemUtils/Source/OnlineSubsystemUtils/Public/OnlineBeaconHost.h
+++ b/Engine/Plugins/Online/OnlineSubsystemUtils/Source/OnlineSubsystemUtils/Public/OnlineBeaconHost.h
@@ -167,6 +167,15 @@ private:
 	/** Connection states used to check against misbehaving connections. */
 	struct FConnectionState
 	{
+		FConnectionState(const AOnlineBeaconHost& InBeaconHost) :
+			BeaconHost(InBeaconHost)
+		{
+		}
+
+		~FConnectionState();
+
+		const AOnlineBeaconHost& BeaconHost;
+		FTimerHandle FinishHandshakeTimerHandle;
 		bool bHasSentHello = false;
 		bool bHasSentChallenge = false;
 		bool bHasSentLogin = false;
@@ -188,6 +197,7 @@ private:
 	void OnConnectionClosed(UNetConnection* Connection);
 
 	bool HandleControlMessage(UNetConnection* Connection, uint8 MessageType, FInBunch& Bunch);
+	void FinishHandshake(UNetConnection* Connection, FString BeaconType);
 	void SendFailurePacket(UNetConnection* Connection, FNetCloseResult&& CloseReason, const FText& ErrorText);
 
 	UE_DEPRECATED(5.1, "SendFailurePacket without CloseReason is deprecated. Use the version which takes CloseReason.")
