Fix non-unity comple error for FTimerDelegate.

diff --git a/Engine/Plugins/Online/OnlineSubsystemUtils/Source/OnlineSubsystemUtils/Private/OnlineBeaconHost.cpp b/Engine/Plugins/Online/OnlineSubsystemUtils/Source/OnlineSubsystemUtils/Private/OnlineBeaconHost.cpp
--- a/Engine/Plugins/Online/OnlineSubsystemUtils/Source/OnlineSubsystemUtils/Private/OnlineBeaconHost.cpp
+++ b/Engine/Plugins/Online/OnlineSubsystemUtils/Source/OnlineSubsystemUtils/Private/OnlineBeaconHost.cpp
@@ -12,6 +12,7 @@
 #include "Misc/NetworkVersion.h"
 #include "Net/DataChannel.h"
 #include "OnlineBeaconClient.h"
+#include "TimerManager.h"
 
 #include UE_INLINE_GENERATED_CPP_BY_NAME(OnlineBeaconHost)
 
