- Use TWeakObjectPtr to check for validity of synth component, as the VoiceEngine/RemoteTalkerData can outlive the UWorld the component is created on.

diff --git a/Engine/Plugins/Online/OnlineSubsystemUtils/Source/OnlineSubsystemUtils/Private/VoiceEngineImpl.cpp b/Engine/Plugins/Online/OnlineSubsystemUtils/Source/OnlineSubsystemUtils/Private/VoiceEngineImpl.cpp
--- a/Engine/Plugins/Online/OnlineSubsystemUtils/Source/OnlineSubsystemUtils/Private/VoiceEngineImpl.cpp
+++ b/Engine/Plugins/Online/OnlineSubsystemUtils/Source/OnlineSubsystemUtils/Private/VoiceEngineImpl.cpp
@@ -122,7 +122,7 @@ void FRemoteTalkerDataImpl::Reset()
 	LastSeen = MAX_FLT;
 	NumFramesStarved = 0;
 
-	if (UObjectInitialized() && VoipSynthComponent)
+	if (UObjectInitialized() && VoipSynthComponent.IsValid())
 	{
 		VoipSynthComponent->Stop();
 
@@ -159,7 +159,7 @@ void FRemoteTalkerDataImpl::Reset()
 
 void FRemoteTalkerDataImpl::Cleanup()
 {
-	if (VoipSynthComponent)
+	if (VoipSynthComponent.IsValid())
 	{
 		VoipSynthComponent->Stop();
 		bIsActive = false;
@@ -561,11 +561,11 @@ uint32 FVoiceEngineImpl::SubmitRemoteVoiceData(const FUniqueNetIdWrapper& Remote
 
 	bool bAudioComponentCreated = false;
 	// Generate a streaming wave audio component for voice playback
-	if (!IsValid(QueuedData.VoipSynthComponent))
+	if (!QueuedData.VoipSynthComponent.IsValid())
 	{
 		CreateSerializeHelper();
 		
-		if (IsValid(QueuedData.VoipSynthComponent))
+		if (QueuedData.VoipSynthComponent.IsValid())
 		{
 			QueuedData.VoipSynthComponent->Stop();
 			QueuedData.VoipSynthComponent->ClosePacketStream();
@@ -579,7 +579,7 @@ uint32 FVoiceEngineImpl::SubmitRemoteVoiceData(const FUniqueNetIdWrapper& Remote
 			}
 		}
 
-		if (QueuedData.VoipSynthComponent)
+		if (QueuedData.VoipSynthComponent.IsValid())
 		{
 			//TODO, make buffer size and buffering delay runtime-controllable parameters.
 			QueuedData.bIsActive = false;
@@ -589,7 +589,7 @@ uint32 FVoiceEngineImpl::SubmitRemoteVoiceData(const FUniqueNetIdWrapper& Remote
 		}
 	}
 
-	if (QueuedData.VoipSynthComponent != nullptr)
+	if (QueuedData.VoipSynthComponent.IsValid())
 	{
 		if (!QueuedData.bIsActive)
 		{
@@ -601,7 +601,7 @@ uint32 FVoiceEngineImpl::SubmitRemoteVoiceData(const FUniqueNetIdWrapper& Remote
 
 			GetVoiceSettingsOverride(RemoteTalkerId, InSettings);
 
-			ApplyVoiceSettings(QueuedData.VoipSynthComponent, InSettings);
+			ApplyVoiceSettings(QueuedData.VoipSynthComponent.Get(), InSettings);
 
 			QueuedData.VoipSynthComponent->ResetBuffer(InSampleCount, UVOIPStatics::GetBufferingDelay());
 			QueuedData.VoipSynthComponent->Start();
@@ -640,7 +640,7 @@ void FVoiceEngineImpl::TickTalkers(float DeltaTime)
 		FRemoteTalkerDataImpl& RemoteData = It.Value();
 		double TimeSince = CurTime - RemoteData.LastSeen;
 
-		if (RemoteData.VoipSynthComponent && RemoteData.VoipSynthComponent->IsIdling() && RemoteData.bIsActive)
+		if (RemoteData.VoipSynthComponent.IsValid() && RemoteData.VoipSynthComponent->IsIdling() && RemoteData.bIsActive)
 		{
 			RemoteData.Reset();
 		}
@@ -698,7 +698,7 @@ void FVoiceEngineImpl::OnAudioFinished()
 	for (FRemoteTalkerData::TIterator It(RemoteTalkerBuffers); It; ++It)
 	{
 		FRemoteTalkerDataImpl& RemoteData = It.Value();
-		if (RemoteData.VoipSynthComponent && RemoteData.VoipSynthComponent->IsIdling())
+		if (RemoteData.VoipSynthComponent.IsValid() && RemoteData.VoipSynthComponent->IsIdling())
 		{
 			UE_LOG_ONLINE_VOICEENGINE(Log, TEXT("Removing VOIP AudioComponent for Id: %s"), *It.Key().ToDebugString());
 			RemoteData.VoipSynthComponent->Stop();
diff --git a/Engine/Plugins/Online/OnlineSubsystemUtils/Source/OnlineSubsystemUtils/Public/VoiceEngineImpl.h b/Engine/Plugins/Online/OnlineSubsystemUtils/Source/OnlineSubsystemUtils/Public/VoiceEngineImpl.h
--- a/Engine/Plugins/Online/OnlineSubsystemUtils/Source/OnlineSubsystemUtils/Public/VoiceEngineImpl.h
+++ b/Engine/Plugins/Online/OnlineSubsystemUtils/Source/OnlineSubsystemUtils/Public/VoiceEngineImpl.h
@@ -82,7 +82,7 @@ public:
 	/** Number of frames starved of audio */
 	int32 NumFramesStarved;
 	/** Synth component playing this buffer (only valid on remote instances) */
-	UVoipListenerSynthComponent* VoipSynthComponent;
+	TWeakObjectPtr<UVoipListenerSynthComponent> VoipSynthComponent;
 	/** Cached Talker Ptr. Is checked against map before use to ensure it has not been destroyed. */
 	UVOIPTalker* CachedTalkerPtr;
 	/** Boolean used to ensure that we only bind the VOIP talker to the SynthComponent's corresponding envelope delegate once. */
@@ -158,9 +158,9 @@ class ONLINESUBSYSTEMUTILS_API FVoiceEngineImpl : public IVoiceEngine, public FS
 			for (FRemoteTalkerData::TIterator It(VoiceEngine->RemoteTalkerBuffers); It; ++It)
 			{
 				FRemoteTalkerDataImpl& RemoteData = It.Value();
-				if (RemoteData.VoipSynthComponent)
+				if (UVoipListenerSynthComponent* VoipSynthComponent = RemoteData.VoipSynthComponent.GetEvenIfUnreachable())
 				{
-					Collector.AddReferencedObject(RemoteData.VoipSynthComponent);
+					Collector.AddReferencedObject(VoipSynthComponent);
 				}
 			}
 		}
