- Delete deprecated OSSes

diff --git a/Engine/Plugins/Online/OnlineFramework/Source/Party/Private/Interactions/CoreInteractions.cpp b/Engine/Plugins/Online/OnlineFramework/Source/Party/Private/Interactions/CoreInteractions.cpp
--- a/Engine/Plugins/Online/OnlineFramework/Source/Party/Private/Interactions/CoreInteractions.cpp
+++ b/Engine/Plugins/Online/OnlineFramework/Source/Party/Private/Interactions/CoreInteractions.cpp
@@ -43,13 +43,7 @@ void FSocialInteraction_AddFriend::ExecuteInteraction(USocialUser& User)
 FText FSocialInteraction_AddPlatformFriend::GetDisplayName(const USocialUser& User)
 {
 	const FName PlatformOssName = USocialManager::GetSocialOssName(ESocialSubsystem::Platform);
-PRAGMA_DISABLE_DEPRECATION_WARNINGS
-	if (PlatformOssName == LIVE_SUBSYSTEM)
-	{
-		return LOCTEXT("AddPlatformFriend_Live", "Add Xbox Live Friend");
-	}
-PRAGMA_ENABLE_DEPRECATION_WARNINGS
-	else if (PlatformOssName == PS4_SUBSYSTEM || USocialSettings::IsSonyOSS(PlatformOssName))
+	if (PlatformOssName == PS4_SUBSYSTEM || USocialSettings::IsSonyOSS(PlatformOssName))
 	{
 		return LOCTEXT("AddPlatformFriend_PSN", "Add Playstation Network Friend");
 	}
diff --git a/Engine/Plugins/Online/OnlineFramework/Source/Party/Private/User/SocialUser.cpp b/Engine/Plugins/Online/OnlineFramework/Source/Party/Private/User/SocialUser.cpp
--- a/Engine/Plugins/Online/OnlineFramework/Source/Party/Private/User/SocialUser.cpp
+++ b/Engine/Plugins/Online/OnlineFramework/Source/Party/Private/User/SocialUser.cpp
@@ -883,12 +883,10 @@ bool USocialUser::CanSendFriendInvite(ESocialSubsystem SubsystemType) const
 	{
 		//@todo DanH: Really need OssCaps or something to be able to just ask an OSS if it supports a given feature. For now, we just magically know that we only support sending XB, PSN, and WeGame invites
 		const FName PlatformOssName = USocialManager::GetSocialOssName(ESocialSubsystem::Platform);
-PRAGMA_DISABLE_DEPRECATION_WARNINGS
-		if (PlatformOssName != LIVE_SUBSYSTEM && PlatformOssName != PS4_SUBSYSTEM && !USocialSettings::IsSonyOSS(PlatformOssName) && PlatformOssName != TENCENT_SUBSYSTEM)
+		if (PlatformOssName != PS4_SUBSYSTEM && !USocialSettings::IsSonyOSS(PlatformOssName) && PlatformOssName != TENCENT_SUBSYSTEM)
 		{
 			return false;
 		}
-PRAGMA_ENABLE_DEPRECATION_WARNINGS
 	}
 
 	return HasSubsystemInfo(SubsystemType) && !IsLocalUser() && !IsFriend(SubsystemType) && !IsBlocked(SubsystemType) && !IsFriendshipPending(SubsystemType);
diff --git a/Engine/Plugins/Online/OnlineSubsystem/Source/Private/OnlineSubsystem.cpp b/Engine/Plugins/Online/OnlineSubsystem/Source/Private/OnlineSubsystem.cpp
--- a/Engine/Plugins/Online/OnlineSubsystem/Source/Private/OnlineSubsystem.cpp
+++ b/Engine/Plugins/Online/OnlineSubsystem/Source/Private/OnlineSubsystem.cpp
@@ -60,9 +60,6 @@ ONLINESUBSYSTEM_API DEFINE_STAT(STAT_Voice_Interface);
 #endif
 
 PRAGMA_DISABLE_DEPRECATION_WARNINGS
-FName WECHAT_SUBSYSTEM(TEXT("WeChat"));
-FName LIVE_SUBSYSTEM(TEXT("LIVE"));
-FName LIVESERVER_SUBSYSTEM(TEXT("LIVESERVER"));
 FName OCULUS_SUBSYSTEM(TEXT("OCULUS"));
 PRAGMA_ENABLE_DEPRECATION_WARNINGS
 
diff --git a/Engine/Plugins/Online/OnlineSubsystem/Source/Public/OnlineSubsystemNames.h b/Engine/Plugins/Online/OnlineSubsystem/Source/Public/OnlineSubsystemNames.h
--- a/Engine/Plugins/Online/OnlineSubsystem/Source/Public/OnlineSubsystemNames.h
+++ b/Engine/Plugins/Online/OnlineSubsystem/Source/Public/OnlineSubsystemNames.h
@@ -77,10 +77,6 @@
 #define TENCENT_SUBSYSTEM FName(TEXT("TENCENT"))
 #endif
 
-UE_DEPRECATED(4.27, "OnlineSubsystemWeChat has been deprecated and will be removed") ONLINESUBSYSTEM_API extern FName WECHAT_SUBSYSTEM;
-UE_DEPRECATED(4.27, "OnlineSubsystemLive has been deprecated and will be removed") ONLINESUBSYSTEM_API extern FName LIVE_SUBSYSTEM;
-UE_DEPRECATED(4.27, "OnlineSubsystemLiveServer has been deprecated and will be removed") ONLINESUBSYSTEM_API extern FName LIVESERVER_SUBSYSTEM;
-
 #ifndef SWITCH_SUBSYSTEM
 #define SWITCH_SUBSYSTEM FName(TEXT("SWITCH"))
 #endif
diff --git a/Engine/Plugins/Online/OnlineSubsystemEOS/Source/OnlineSubsystemEOS/Private/UserManagerEOS.cpp b/Engine/Plugins/Online/OnlineSubsystemEOS/Source/OnlineSubsystemEOS/Private/UserManagerEOS.cpp
--- a/Engine/Plugins/Online/OnlineSubsystemEOS/Source/OnlineSubsystemEOS/Private/UserManagerEOS.cpp
+++ b/Engine/Plugins/Online/OnlineSubsystemEOS/Source/OnlineSubsystemEOS/Private/UserManagerEOS.cpp
@@ -106,12 +106,10 @@ static inline EOS_EExternalCredentialType ToEOS_EExternalCredentialType(FName OS
 	{
 		return EOS_EExternalCredentialType::EOS_ECT_PSN_ID_TOKEN;
 	}
-PRAGMA_DISABLE_DEPRECATION_WARNINGS
-	else if (OSSName == LIVE_SUBSYSTEM || USE_XBL_XSTS_TOKEN)
+	else if (USE_XBL_XSTS_TOKEN)
 	{
 		return EOS_EExternalCredentialType::EOS_ECT_XBL_XSTS_TOKEN;
 	}
-PRAGMA_ENABLE_DEPRECATION_WARNINGS
 	else if (OSSName == SWITCH_SUBSYSTEM)
 	{
 		if (AccountCredentials.Type == TEXT("NintendoAccount"))
diff --git a/Engine/Plugins/Online/OnlineSubsystemEOS/Source/OnlineSubsystemEOSPlus/Private/OnlineUserEOSPlus.cpp b/Engine/Plugins/Online/OnlineSubsystemEOS/Source/OnlineSubsystemEOSPlus/Private/OnlineUserEOSPlus.cpp
--- a/Engine/Plugins/Online/OnlineSubsystemEOS/Source/OnlineSubsystemEOSPlus/Private/OnlineUserEOSPlus.cpp
+++ b/Engine/Plugins/Online/OnlineSubsystemEOS/Source/OnlineSubsystemEOSPlus/Private/OnlineUserEOSPlus.cpp
@@ -12,7 +12,6 @@ enum class EOSSValue : uint8
 	Null,
 	Steam,
 	PS4,
-	XboxLive,
 	Switch,
 	Apple
 };
@@ -27,12 +26,6 @@ static inline EOSSValue ToEOSSValue(FName OSSName)
 	{
 		return EOSSValue::PS4;
 	}
-PRAGMA_DISABLE_DEPRECATION_WARNINGS
-	else if (OSSName == LIVE_SUBSYSTEM)
-	{
-		return EOSSValue::XboxLive;
-	}
-PRAGMA_ENABLE_DEPRECATION_WARNINGS
 	else if (OSSName == SWITCH_SUBSYSTEM)
 	{
 		return EOSSValue::Switch;
@@ -54,12 +47,6 @@ static inline FName ToOSSName(EOSSValue OSSValue)
 	{
 		return PS4_SUBSYSTEM;
 	}
-PRAGMA_DISABLE_DEPRECATION_WARNINGS
-	else if (OSSValue == EOSSValue::XboxLive)
-	{
-		return LIVE_SUBSYSTEM;
-	}
-PRAGMA_ENABLE_DEPRECATION_WARNINGS
 	else if (OSSValue == EOSSValue::Switch)
 	{
 		return SWITCH_SUBSYSTEM;
diff --git a/Engine/Plugins/Online/OnlineSubsystemUtils/Source/OnlineSubsystemUtils/Private/OnlineEngineInterfaceImpl.cpp b/Engine/Plugins/Online/OnlineSubsystemUtils/Source/OnlineSubsystemUtils/Private/OnlineEngineInterfaceImpl.cpp
--- a/Engine/Plugins/Online/OnlineSubsystemUtils/Source/OnlineSubsystemUtils/Private/OnlineEngineInterfaceImpl.cpp
+++ b/Engine/Plugins/Online/OnlineSubsystemUtils/Source/OnlineSubsystemUtils/Private/OnlineEngineInterfaceImpl.cpp
@@ -97,13 +97,7 @@ FName UOnlineEngineInterfaceImpl::GetDedicatedServerSubsystemNameForSubsystem(co
 	// then the server can pass the value of that option to this function to get the name of
 	// the corresponding server OSS for that client, if one exists.
 
-PRAGMA_DISABLE_DEPRECATION_WARNINGS
-	if (Subsystem == LIVE_SUBSYSTEM)
-	{
-		return LIVESERVER_SUBSYSTEM;
-	}
-PRAGMA_ENABLE_DEPRECATION_WARNINGS
-	else if (Subsystem == PS4_SUBSYSTEM)
+	if (Subsystem == PS4_SUBSYSTEM)
 	{
 		return PS4SERVER_SUBSYSTEM;
 	}
diff --git a/Engine/Plugins/Online/OnlineSubsystemUtils/Source/OnlineSubsystemUtils/Private/OnlineSubsystemUtilsModule.cpp b/Engine/Plugins/Online/OnlineSubsystemUtils/Source/OnlineSubsystemUtils/Private/OnlineSubsystemUtilsModule.cpp
--- a/Engine/Plugins/Online/OnlineSubsystemUtils/Source/OnlineSubsystemUtils/Private/OnlineSubsystemUtilsModule.cpp
+++ b/Engine/Plugins/Online/OnlineSubsystemUtils/Source/OnlineSubsystemUtils/Private/OnlineSubsystemUtilsModule.cpp
@@ -195,9 +195,6 @@ private:
 		CREATE_HASH(MCP_SUBSYSTEM);
 		CREATE_HASH(STEAM_SUBSYSTEM);
 		CREATE_HASH(PS4_SUBSYSTEM);
-PRAGMA_DISABLE_DEPRECATION_WARNINGS
-		CREATE_HASH(LIVE_SUBSYSTEM);
-PRAGMA_ENABLE_DEPRECATION_WARNINGS
 		CREATE_HASH(GOOGLE_SUBSYSTEM);
 		CREATE_HASH(GOOGLEPLAY_SUBSYSTEM);
 		CREATE_HASH(FACEBOOK_SUBSYSTEM);
@@ -209,16 +206,12 @@ PRAGMA_ENABLE_DEPRECATION_WARNINGS
 		CREATE_HASH(GAMECIRCLE_SUBSYSTEM);
 		CREATE_HASH(THUNDERHEAD_SUBSYSTEM); 
 PRAGMA_DISABLE_DEPRECATION_WARNINGS
-		CREATE_HASH(WECHAT_SUBSYSTEM);
 		CREATE_HASH(OCULUS_SUBSYSTEM); 
 PRAGMA_ENABLE_DEPRECATION_WARNINGS
 
 		CREATE_HASH(QUAIL_SUBSYSTEM);
 		// Shouldn't need these as they are mocking interfaces for existing platforms
 		CREATE_HASH(PS4SERVER_SUBSYSTEM);
-PRAGMA_DISABLE_DEPRECATION_WARNINGS
-		CREATE_HASH(LIVESERVER_SUBSYSTEM);
-PRAGMA_ENABLE_DEPRECATION_WARNINGS
 
 #undef CREATE_HASH
 
diff --git a/Engine/Plugins/Online/OnlineSubsystemUtils/Source/OnlineSubsystemUtils/Private/Tests/OnlineSubsystemUtilitiesTests.cpp b/Engine/Plugins/Online/OnlineSubsystemUtils/Source/OnlineSubsystemUtils/Private/Tests/OnlineSubsystemUtilitiesTests.cpp
--- a/Engine/Plugins/Online/OnlineSubsystemUtils/Source/OnlineSubsystemUtils/Private/Tests/OnlineSubsystemUtilitiesTests.cpp
+++ b/Engine/Plugins/Online/OnlineSubsystemUtils/Source/OnlineSubsystemUtils/Private/Tests/OnlineSubsystemUtilitiesTests.cpp
@@ -51,7 +51,6 @@ namespace OSSUtilsTestHelper
 		OSSList.AddUnique(TEXT( "Steam" ) );
 		OSSList.AddUnique(TEXT( "Facebook" ) );
 		OSSList.AddUnique(TEXT( "IOS" ) );
-		OSSList.AddUnique(TEXT( "WeChat" ) );
 		OSSList.AddUnique(TEXT( "GooglePlay" ) );
 		OSSList.AddUnique(TEXT( "Live" ) );
 		OSSList.AddUnique(TEXT( "Thunderhead" ) );
