Add port binding options for reuse and to exit on failure.

diff --git a/Engine/Plugins/Online/OnlineSubsystemUtils/Source/OnlineSubsystemUtils/Classes/IpNetDriver.h b/Engine/Plugins/Online/OnlineSubsystemUtils/Source/OnlineSubsystemUtils/Classes/IpNetDriver.h
--- a/Engine/Plugins/Online/OnlineSubsystemUtils/Source/OnlineSubsystemUtils/Classes/IpNetDriver.h
+++ b/Engine/Plugins/Online/OnlineSubsystemUtils/Source/OnlineSubsystemUtils/Classes/IpNetDriver.h
@@ -75,6 +75,10 @@ public:
     UPROPERTY(Config)
     uint32 AllowPlayerPortUnreach:1;
 
+	/** Whether to exit if binding a port failed. */
+	UPROPERTY(Config)
+	uint32 bExitOnBindFailure:1;
+
 	/** Number of ports which will be tried if current one is not available for binding (i.e. if told to bind to port N, will try from N to N+MaxPortCountToTry inclusive) */
 	UPROPERTY(Config)
 	uint32 MaxPortCountToTry;
diff --git a/Engine/Plugins/Online/OnlineSubsystemUtils/Source/OnlineSubsystemUtils/Private/IpNetDriver.cpp b/Engine/Plugins/Online/OnlineSubsystemUtils/Source/OnlineSubsystemUtils/Private/IpNetDriver.cpp
--- a/Engine/Plugins/Online/OnlineSubsystemUtils/Source/OnlineSubsystemUtils/Private/IpNetDriver.cpp
+++ b/Engine/Plugins/Online/OnlineSubsystemUtils/Source/OnlineSubsystemUtils/Private/IpNetDriver.cpp
@@ -846,8 +846,15 @@ FUniqueSocket UIpNetDriver::CreateAndBindSocket(TSharedRef<FInternetAddr> BindAd
 	{
 		Error = FString::Printf(TEXT("%s: binding to port %i failed (%i)"), SocketSubsystem->GetSocketAPIName(), AttemptPort,
 			(int32)SocketSubsystem->GetLastErrorCode());
+
+		if (bExitOnBindFailure)
+		{
+			UE_LOG(LogNet, Fatal, TEXT("Fatal error: %s"), *Error);
+		}
+
 		return nullptr;
 	}
+
 	if (NewSocket->SetNonBlocking() == false)
 	{
 		Error = FString::Printf(TEXT("%s: SetNonBlocking failed (%i)"), SocketSubsystem->GetSocketAPIName(),
diff --git a/Engine/Plugins/Online/OnlineSubsystemUtils/Source/OnlineSubsystemUtils/Private/OnlineBeaconHost.cpp b/Engine/Plugins/Online/OnlineSubsystemUtils/Source/OnlineSubsystemUtils/Private/OnlineBeaconHost.cpp
--- a/Engine/Plugins/Online/OnlineSubsystemUtils/Source/OnlineSubsystemUtils/Private/OnlineBeaconHost.cpp
+++ b/Engine/Plugins/Online/OnlineSubsystemUtils/Source/OnlineSubsystemUtils/Private/OnlineBeaconHost.cpp
@@ -96,7 +96,7 @@ bool AOnlineBeaconHost::InitHost()
 		if (InitBase() && NetDriver)
 		{
 			FString Error;
-			if (NetDriver->InitListen(this, URL, false, Error))
+			if (NetDriver->InitListen(this, URL, bReuseAddressAndPort, Error))
 			{
 				ListenPort = URL.Port;
 				NetDriver->SetWorld(GetWorld());
diff --git a/Engine/Plugins/Online/OnlineSubsystemUtils/Source/OnlineSubsystemUtils/Public/OnlineBeaconHost.h b/Engine/Plugins/Online/OnlineSubsystemUtils/Source/OnlineSubsystemUtils/Public/OnlineBeaconHost.h
--- a/Engine/Plugins/Online/OnlineSubsystemUtils/Source/OnlineSubsystemUtils/Public/OnlineBeaconHost.h
+++ b/Engine/Plugins/Online/OnlineSubsystemUtils/Source/OnlineSubsystemUtils/Public/OnlineBeaconHost.h
@@ -34,6 +34,13 @@ public:
 	UPROPERTY(Config)
 	int32 ListenPort;
 
+	/**
+	 * Whether to configure the listening socket to allow reuse of the address and port. If this is true, be sure no other
+	 * servers can run on the same port, otherwise this can lead to undefined behavior since packets will go to two servers.
+	 */
+	UPROPERTY(Config)
+	bool bReuseAddressAndPort = false;
+
 	//~ Begin AActor Interface
 	virtual void OnNetCleanup(UNetConnection* Connection) override;
 	//~ End AActor Interface
diff --git a/Engine/Source/Runtime/Engine/Classes/GameFramework/WorldSettings.h b/Engine/Source/Runtime/Engine/Classes/GameFramework/WorldSettings.h
--- a/Engine/Source/Runtime/Engine/Classes/GameFramework/WorldSettings.h
+++ b/Engine/Source/Runtime/Engine/Classes/GameFramework/WorldSettings.h
@@ -572,6 +572,13 @@ public:
 	uint8 bIncludeGridSizeInNameForPartitionedActors : 1;
 #endif
 
+	/**
+	 * Whether to configure the listening socket to allow reuse of the address and port. If this is true, be sure no other
+	 * servers can run on the same port, otherwise this can lead to undefined behavior since packets will go to two servers.
+	 */
+	UPROPERTY(EditAnywhere, Config, Category=Network)
+	uint8 bReuseAddressAndPort : 1;
+
 	UPROPERTY(EditAnywhere, BlueprintReadOnly, Category=AI, meta=(MetaClass="/Script/Engine.AISystemBase", editcondition="bEnableAISystem"), AdvancedDisplay)
 	TSoftClassPtr<UAISystemBase> AISystemClass;
 
diff --git a/Engine/Source/Runtime/Engine/Private/World.cpp b/Engine/Source/Runtime/Engine/Private/World.cpp
--- a/Engine/Source/Runtime/Engine/Private/World.cpp
+++ b/Engine/Source/Runtime/Engine/Private/World.cpp
@@ -6407,8 +6407,11 @@ bool UWorld::Listen( FURL& InURL )
 		return false;
 	}
 
+	AWorldSettings* WorldSettings = GetWorldSettings();
+	const bool bReuseAddressAndPort = WorldSettings ? WorldSettings->bReuseAddressAndPort : false;
+
 	FString Error;
-	if( !NetDriver->InitListen( this, InURL, false, Error ) )
+	if( !NetDriver->InitListen( this, InURL, bReuseAddressAndPort, Error ) )
 	{
 		GEngine->BroadcastNetworkFailure(this, NetDriver, ENetworkFailure::NetDriverListenFailure, Error);
 		UE_LOG(LogWorld, Log,  TEXT("Failed to listen: %s"), *Error );
diff --git a/Engine/Source/Runtime/Engine/Private/WorldSettings.cpp b/Engine/Source/Runtime/Engine/Private/WorldSettings.cpp
--- a/Engine/Source/Runtime/Engine/Private/WorldSettings.cpp
+++ b/Engine/Source/Runtime/Engine/Private/WorldSettings.cpp
@@ -123,6 +123,8 @@ PRAGMA_ENABLE_DEPRECATION_WARNINGS
 	LastBookmarkClass = DefaultBookmarkClass;
 
 	LevelInstancePivotOffset = FVector::ZeroVector;
+
+	bReuseAddressAndPort = false;
 }
 
 void AWorldSettings::PostInitProperties()
