Fix build break from missing MinimalApi specifier in CL24392821.

#rnx

diff --git a/Engine/Plugins/Online/OnlineSubsystemUtils/Source/OnlineSubsystemUtils/Public/OnlineBeacon.h b/Engine/Plugins/Online/OnlineSubsystemUtils/Source/OnlineSubsystemUtils/Public/OnlineBeacon.h
--- a/Engine/Plugins/Online/OnlineSubsystemUtils/Source/OnlineSubsystemUtils/Public/OnlineBeacon.h
+++ b/Engine/Plugins/Online/OnlineSubsystemUtils/Source/OnlineSubsystemUtils/Public/OnlineBeacon.h
@@ -29,7 +29,7 @@ namespace EBeaconState
 /**
  * Base class for beacon communication (Unreal Networking, but outside normal gameplay traffic)
  */
-UCLASS(transient, config=Engine, notplaceable)
+UCLASS(MinimalApi, transient, config=Engine, notplaceable)
 class AOnlineBeacon : public AActor, public FNetworkNotify
 {
 	GENERATED_UCLASS_BODY()
