Add CVar for specifying which beacon type to apply handshake timeout to.

diff --git a/Engine/Plugins/Online/OnlineSubsystemUtils/Source/OnlineSubsystemUtils/Private/OnlineBeaconHost.cpp b/Engine/Plugins/Online/OnlineSubsystemUtils/Source/OnlineSubsystemUtils/Private/OnlineBeaconHost.cpp
--- a/Engine/Plugins/Online/OnlineSubsystemUtils/Source/OnlineSubsystemUtils/Private/OnlineBeaconHost.cpp
+++ b/Engine/Plugins/Online/OnlineSubsystemUtils/Source/OnlineSubsystemUtils/Private/OnlineBeaconHost.cpp
@@ -23,6 +23,13 @@ namespace BeaconConsoleVariables
 		TEXT("Delay time before finishing handshake by calling client RPC\n")
 		TEXT("Time in seconds. A value of 0 means no delay, and a negative value means never call it."),
 		ECVF_Default);
+
+	TAutoConsoleVariable<FString> CVarDelayFinishHandshakeType(
+		TEXT("beacon.DelayFinishHandshakeBeaconType"),
+		TEXT(""),
+		TEXT("The type of beacon to apply the handshake delay to.\n")
+		TEXT("Leave blank for all."),
+		ECVF_Default);
 }
 #endif
 
@@ -468,22 +475,26 @@ bool AOnlineBeaconHost::HandleControlMessage(UNetConnection* Connection, uint8 M
 
 #if UE_BUILD_SHIPPING
 			const float Delay = 0.0f;
+			const FString DelayBeaconType;
+			const bool bDelayAllowedForBeaconType = false;
 #else
 			const float Delay = BeaconConsoleVariables::CVarDelayFinishHandshake.GetValueOnGameThread();
+			const FString DelayBeaconType = BeaconConsoleVariables::CVarDelayFinishHandshakeType.GetValueOnGameThread();
+			const bool bDelayAllowedForBeaconType = DelayBeaconType.IsEmpty() || DelayBeaconType == BeaconType;
 #endif
-			if (Delay == 0.0f)
+			if (Delay == 0.0f || !bDelayAllowedForBeaconType)
 			{
 				FinishHandshake(Connection, MoveTemp(BeaconType));
 			}
 			else if (Delay > 0.0f)
 			{
-				UE_LOG(LogBeacon, Verbose, TEXT("%s: Delay for handshake completion: %.3f"), *GetDebugName(Connection), Delay);
+				UE_LOG(LogBeacon, Verbose, TEXT("%s: BeaconType: %s, Delay for handshake completion: %.3f"), *GetDebugName(Connection), *BeaconType, Delay);
 				FTimerDelegate TimerDelegate = FTimerDelegate::CreateUObject(this, &AOnlineBeaconHost::FinishHandshake, Connection, MoveTemp(BeaconType));
 				GetWorldTimerManager().SetTimer(ConnState->FinishHandshakeTimerHandle, TimerDelegate, Delay, false);
 			}
 			else
 			{
-				UE_LOG(LogBeacon, Verbose, TEXT("%s: Handshake will never complete, client will either timeout or hang if timeouts are disabled."), *GetDebugName(Connection));
+				UE_LOG(LogBeacon, Verbose, TEXT("%s: BeaconType: %s, Handshake will never complete, client will either timeout or hang if timeouts are disabled."), *GetDebugName(Connection), *BeaconType);
 			}
 		}
 
