Add commands for dumping and disconnecting online beacons.

diff --git a/Engine/Plugins/Online/OnlineSubsystemUtils/Source/OnlineSubsystemUtils/Private/OnlineSubsystemUtils.cpp b/Engine/Plugins/Online/OnlineSubsystemUtils/Source/OnlineSubsystemUtils/Private/OnlineSubsystemUtils.cpp
--- a/Engine/Plugins/Online/OnlineSubsystemUtils/Source/OnlineSubsystemUtils/Private/OnlineSubsystemUtils.cpp
+++ b/Engine/Plugins/Online/OnlineSubsystemUtils/Source/OnlineSubsystemUtils/Private/OnlineSubsystemUtils.cpp
@@ -1,6 +1,7 @@
 // Copyright Epic Games, Inc. All Rights Reserved.
 
 #include "OnlineSubsystemUtils.h"
+#include "Algo/Copy.h"
 #include "Logging/LogScopedVerbosityOverride.h"
 #include "Misc/ConfigCacheIni.h"
 #include "GameFramework/PlayerState.h"
@@ -10,6 +11,8 @@
 #include "SocketSubsystem.h"
 #include "OnlineSubsystemImpl.h"
 #include "OnlineSubsystemBPCallHelper.h"
+#include "OnlineBeaconClient.h"
+#include "UObject/UObjectHash.h"
 
 #include "Sound/SoundGroups.h"
 #include "VoiceModule.h"
@@ -432,6 +435,180 @@ bool HandleVoiceCommands(IOnlineSubsystem* InOnlineSub, const UWorld* InWorld, c
 	return bWasHandled;
 }
 
+enum class FOnlineBeaconClientFilterFlags : uint8
+{
+	None = 0,
+
+	/** Whether to match all beacons when no filter is set. */
+	WildcardAllowed,
+};
+ENUM_CLASS_FLAGS(FOnlineBeaconClientFilterFlags);
+
+class FOnlineBeaconClientFilter final
+{
+public:
+	static bool Parse(const TCHAR* Cmd, FOnlineBeaconClientFilter& OutFilter, FOnlineBeaconClientFilterFlags Flags);
+
+	FOnlineBeaconClientFilter() = default;
+	FOnlineBeaconClientFilter(const FString& Name, const FString& Type, const FUniqueNetIdRepl& UniqueId, FOnlineBeaconClientFilterFlags Flags);
+	FString ToString() const;
+
+	// Returns true if the beacon matches the filter.
+	bool operator()(const UObject* Obj) const;
+	bool operator()(const AOnlineBeaconClient* Beacon) const;
+
+private:
+	FString Name;
+	FString Type;
+	FUniqueNetIdRepl UniqueId;
+	bool bIsWildcard = false;
+	FOnlineBeaconClientFilterFlags Flags;
+};
+
+bool FOnlineBeaconClientFilter::Parse(const TCHAR* Cmd, FOnlineBeaconClientFilter& OutFilter, FOnlineBeaconClientFilterFlags Flags)
+{
+	bool bArgumentsValid = true;
+
+	FString BeaconName;
+	FParse::Value(Cmd, TEXT("Name="), BeaconName);
+
+	FString BeaconType;
+	FParse::Value(Cmd, TEXT("Type="), BeaconType);
+
+	FUniqueNetIdRepl UniqueId;
+	FString UniqueIdStr;
+	if (FParse::Value(Cmd, TEXT("UniqueId="), UniqueIdStr))
+	{
+		UniqueId.FromJson(UniqueIdStr);
+		if (!UniqueId.IsValid())
+		{
+			UE_LOG(LogBeacon, Warning, TEXT("Beacon filter parse failure - unable to parse UniqueId: %s"), *UniqueIdStr);
+			bArgumentsValid = false;
+		}
+	}
+
+	if (!bArgumentsValid)
+	{
+		return false;
+	}
+
+	OutFilter = FOnlineBeaconClientFilter(BeaconName, BeaconType, UniqueId, Flags);
+	return true;
+}
+
+FOnlineBeaconClientFilter::FOnlineBeaconClientFilter(const FString& Name, const FString& Type, const FUniqueNetIdRepl& UniqueId, FOnlineBeaconClientFilterFlags Flags)
+	: Name(Name)
+	, Type(Type)
+	, UniqueId(UniqueId)
+	, Flags(Flags)
+{
+	if (EnumHasAnyFlags(Flags, FOnlineBeaconClientFilterFlags::WildcardAllowed))
+	{
+		bIsWildcard = Name.IsEmpty() && Type.IsEmpty() && !UniqueId.IsValid();
+	}
+}
+
+FString FOnlineBeaconClientFilter::ToString() const
+{
+	if (bIsWildcard)
+	{
+		return FString(TEXT("Wildcard"));
+	}
+	else
+	{
+		return FString::Printf(TEXT("Name: %s, Type: %s, UniqueId: %s"),
+			!Name.IsEmpty() ? *Name : TEXT("<empty>"),
+			!Type.IsEmpty() ? *Type : TEXT("<empty>"),
+			*UniqueId.ToDebugString());
+	}
+}
+
+bool FOnlineBeaconClientFilter::operator()(const UObject* Obj) const
+{
+	return (*this)(Cast<AOnlineBeaconClient>(Obj));
+}
+
+bool FOnlineBeaconClientFilter::operator()(const AOnlineBeaconClient* Beacon) const
+{
+	return IsValid(Beacon) && (bIsWildcard || Name == Beacon->GetName() || Type == Beacon->GetBeaconType() || UniqueId == Beacon->GetUniqueId());
+}
+
+TArray<AOnlineBeaconClient*> GetFilteredOnlineBeaconClientInstances(const FOnlineBeaconClientFilter& Filter)
+{
+	TArray<UObject*> Objects;
+	GetObjectsOfClass(AOnlineBeaconClient::StaticClass(), Objects);
+
+	TArray<AOnlineBeaconClient*> ClientBeacons;
+	Algo::TransformIf(Objects, ClientBeacons, Filter,
+		[](UObject* Obj)->AOnlineBeaconClient* { return Cast<AOnlineBeaconClient>(Obj); });
+
+	return ClientBeacons;
+}
+
+FString GetBeaconLogString(const AOnlineBeaconClient& ClientBeacon)
+{
+	return FString::Printf(TEXT("Name: %s, Type: %s, UniqueId: %s, State: %s"), *ClientBeacon.GetName(), *ClientBeacon.GetBeaconType(), *ClientBeacon.GetUniqueId().ToDebugString(), *UEnum::GetValueAsString(ClientBeacon.GetConnectionState()));
+}
+
+bool HandleBeaconCommands(IOnlineSubsystem* InOnlineSub, const UWorld* InWorld, const TCHAR* Cmd, FOutputDevice& Ar)
+{
+	bool bWasHandled = true;
+
+	if (FParse::Command(&Cmd, TEXT("DUMP")))
+	{
+		FOnlineBeaconClientFilter BeaconFilter;
+		if (FOnlineBeaconClientFilter::Parse(Cmd, BeaconFilter, FOnlineBeaconClientFilterFlags::WildcardAllowed))
+		{
+			UE_LOG(LogBeacon, Log, TEXT("Dumping beacon info based on filter: %s."), *BeaconFilter.ToString());
+			TArray<AOnlineBeaconClient*> FilteredBeacons = GetFilteredOnlineBeaconClientInstances(BeaconFilter);
+			for (AOnlineBeaconClient* ClientBeacon : FilteredBeacons)
+			{
+				UE_LOG(LogBeacon, Log, TEXT("    AOnlineBeaconClient: %s"), *GetBeaconLogString(*ClientBeacon));
+			}
+			UE_LOG(LogBeacon, Log, TEXT("Dump complete. Found %d beacons."), FilteredBeacons.Num());
+		}
+		else
+		{
+			UE_LOG(LogBeacon, Warning, TEXT("DUMP command failure - unable to parse arguments."));
+		}
+	}
+	else if (FParse::Command(&Cmd, TEXT("DISCONNECT")))
+	{
+		FOnlineBeaconClientFilter BeaconFilter;
+		if (FOnlineBeaconClientFilter::Parse(Cmd, BeaconFilter, FOnlineBeaconClientFilterFlags::None))
+		{
+			UE_LOG(LogBeacon, Log, TEXT("Disconnecting beacons based on filter: %s."), *BeaconFilter.ToString());
+
+			int32 NumDisconnected = 0;
+			for (AOnlineBeaconClient* ClientBeacon : GetFilteredOnlineBeaconClientInstances(BeaconFilter))
+			{
+				UNetConnection* Connection = ClientBeacon->GetNetConnection();
+
+				// Closing the connection will start the chain of events leading to the removal from lists and destruction of the actor
+				if (Connection && Connection->GetConnectionState() != USOCK_Closed)
+				{
+					UE_LOG(LogBeacon, Log, TEXT("    Disconnecting beacon: %s"), *GetBeaconLogString(*ClientBeacon));
+					Connection->FlushNet(true);
+					Connection->Close();
+					++NumDisconnected;
+				}
+				else
+				{
+					UE_LOG(LogBeacon, Log, TEXT("    Beacon not in a valid state to disconnect: %s"), *GetBeaconLogString(*ClientBeacon));
+				}
+			}
+
+			UE_LOG(LogBeacon, Log, TEXT("Disconnect complete. %d beacons affected."), NumDisconnected);
+		}
+		else
+		{
+			UE_LOG(LogBeacon, Warning, TEXT("DISCONNECT command failure - unable to parse arguments."));
+		}
+	}
+
+	return bWasHandled;
+}
+
 /**
  * Exec handler that routes online specific execs to the proper subsystem
  *
@@ -760,6 +937,10 @@ static bool OnlineExec( UWorld* InWorld, const TCHAR* Cmd, FOutputDevice& Ar )
 				{
 					bWasHandled = HandleVoiceCommands(OnlineSub, InWorld, Cmd, Ar);
 				}
+				else if (FParse::Command(&Cmd, TEXT("BEACON")))
+				{
+					bWasHandled = HandleBeaconCommands(OnlineSub, InWorld, Cmd, Ar);
+				}
 			}
 		}
 	}
