Correctly handle multiple filter parameters in beacon commands.

diff --git a/Engine/Plugins/Online/OnlineSubsystemUtils/Source/OnlineSubsystemUtils/Private/OnlineSubsystemUtils.cpp b/Engine/Plugins/Online/OnlineSubsystemUtils/Source/OnlineSubsystemUtils/Private/OnlineSubsystemUtils.cpp
--- a/Engine/Plugins/Online/OnlineSubsystemUtils/Source/OnlineSubsystemUtils/Private/OnlineSubsystemUtils.cpp
+++ b/Engine/Plugins/Online/OnlineSubsystemUtils/Source/OnlineSubsystemUtils/Private/OnlineSubsystemUtils.cpp
@@ -450,7 +450,6 @@ public:
 	static bool Parse(const TCHAR* Cmd, FOnlineBeaconClientFilter& OutFilter, FOnlineBeaconClientFilterFlags Flags);
 
 	FOnlineBeaconClientFilter() = default;
-	FOnlineBeaconClientFilter(const FString& Name, const FString& Type, const FUniqueNetIdRepl& UniqueId, FOnlineBeaconClientFilterFlags Flags);
 	FString ToString() const;
 
 	// Returns true if the beacon matches the filter.
@@ -458,6 +457,8 @@ public:
 	bool operator()(const AOnlineBeaconClient* Beacon) const;
 
 private:
+	FOnlineBeaconClientFilter(const FString& Name, const FString& Type, const FUniqueNetIdRepl& UniqueId, FOnlineBeaconClientFilterFlags Flags);
+
 	FString Name;
 	FString Type;
 	FUniqueNetIdRepl UniqueId;
@@ -487,6 +488,13 @@ bool FOnlineBeaconClientFilter::Parse(const TCHAR* Cmd, FOnlineBeaconClientFilte
 		}
 	}
 
+	const bool bHasAnyArguments = !BeaconName.IsEmpty() || !BeaconType.IsEmpty() || UniqueId.IsValid();
+	if (!EnumHasAnyFlags(Flags, FOnlineBeaconClientFilterFlags::WildcardAllowed) && !bHasAnyArguments)
+	{
+		UE_LOG(LogBeacon, Warning, TEXT("Beacon filter parse failure - no arguments found. Wildcard filter is not allowed."));
+		bArgumentsValid = false;
+	}
+
 	if (!bArgumentsValid)
 	{
 		return false;
@@ -530,7 +538,15 @@ bool FOnlineBeaconClientFilter::operator()(const UObject* Obj) const
 
 bool FOnlineBeaconClientFilter::operator()(const AOnlineBeaconClient* Beacon) const
 {
-	return IsValid(Beacon) && (bIsWildcard || Name == Beacon->GetName() || Type == Beacon->GetBeaconType() || UniqueId == Beacon->GetUniqueId());
+	if (!IsValid(Beacon))
+	{
+		return false;
+	}
+
+	const bool bNameMatches = Name.IsEmpty() || Name == Beacon->GetName();
+	const bool bTypeMatches = Type.IsEmpty() || Type == Beacon->GetBeaconType();
+	const bool bUniqueIdMatches = !UniqueId.IsValid() || UniqueId == Beacon->GetUniqueId();
+	return bIsWildcard || (bNameMatches && bTypeMatches && bUniqueIdMatches);
 }
 
 TArray<AOnlineBeaconClient*> GetFilteredOnlineBeaconClientInstances(const FOnlineBeaconClientFilter& Filter)
