Modify OnlineBeaconClient login parameters to allow for more parameters than just the auth ticket

diff --git a/Engine/Plugins/Online/OnlineSubsystemUtils/Source/OnlineSubsystemUtils/Private/OnlineBeaconClient.cpp b/Engine/Plugins/Online/OnlineSubsystemUtils/Source/OnlineSubsystemUtils/Private/OnlineBeaconClient.cpp
--- a/Engine/Plugins/Online/OnlineSubsystemUtils/Source/OnlineSubsystemUtils/Private/OnlineBeaconClient.cpp
+++ b/Engine/Plugins/Online/OnlineSubsystemUtils/Source/OnlineSubsystemUtils/Private/OnlineBeaconClient.cpp
@@ -90,6 +90,55 @@ void AOnlineBeaconClient::SetConnectionState(EBeaconConnectionState NewConnectio
 	ConnectionState = NewConnectionState;
 }
 
+FString AOnlineBeaconClient::GetLoginOptions(const FUniqueNetIdRepl& PlayerId)
+{
+	const FString& AuthTicket = GetAuthTicketInternal(*PlayerId);
+
+	FURL URL(nullptr, TEXT(""), TRAVEL_Absolute);
+	if (!AuthTicket.IsEmpty())
+	{
+		URL.AddOption(*FString::Printf(TEXT("AuthTicket=%s"), *AuthTicket));
+	}
+	return URL.ToString();
+}
+
+namespace
+{
+FUniqueNetIdRepl GetDedicatedServerAccountId(UWorld* World)
+{
+	FUniqueNetIdRepl AccountId;
+	IOnlineIdentityPtr IdentityPtr = Online::GetIdentityInterface(World);
+	if (IdentityPtr.IsValid())
+	{
+		AccountId = IdentityPtr->GetUniquePlayerId(DEDICATED_SERVER_USER_INDEX);
+	}
+	return AccountId;
+}
+
+FUniqueNetIdRepl GetFirstPlayerAccountId(UWorld* World)
+{
+	FUniqueNetIdRepl AccountId;
+	// Prefer local player's preferred id
+	if (ULocalPlayer* LocalPlayer = GEngine->GetFirstGamePlayer(World))
+	{
+		AccountId = LocalPlayer->GetPreferredUniqueNetId();
+	}
+	else
+	{
+		// Fall back to querying the identity interface.
+		if (IOnlineIdentityPtr IdentityPtr = Online::GetIdentityInterface(World))
+		{
+			TArray<TSharedPtr<FUserOnlineAccount> > UserAccounts = IdentityPtr->GetAllUserAccounts();
+			if (!UserAccounts.IsEmpty())
+			{
+				AccountId = UserAccounts[0]->GetUserId();
+			}
+		}
+	}
+	return AccountId;
+}
+}
+
 bool AOnlineBeaconClient::InitClient(FURL& URL)
 {
 	bool bSuccess = false;
@@ -109,32 +158,12 @@ bool AOnlineBeaconClient::InitClient(FURL& URL)
 
 					if (IsRunningDedicatedServer())
 					{
-						IOnlineIdentityPtr IdentityPtr = Online::GetIdentityInterface(World);
-						if (IdentityPtr.IsValid())
-						{
-							BeaconConnection->PlayerId = IdentityPtr->GetUniquePlayerId(DEDICATED_SERVER_USER_INDEX);
-						}
+						BeaconConnection->PlayerId = GetDedicatedServerAccountId(World);
 					}
 					else
 					{
-						ULocalPlayer* LocalPlayer = GEngine->GetFirstGamePlayer(World);
-						if (LocalPlayer)
-						{
-							// Send the player unique Id at login
-							BeaconConnection->PlayerId = LocalPlayer->GetPreferredUniqueNetId();
-						}
-						else
-						{
-							// Fall back to querying the identity interface.
-							if (IOnlineIdentityPtr IdentityPtr = Online::GetIdentityInterface(World))
-							{
-								TArray<TSharedPtr<FUserOnlineAccount> > UserAccounts = IdentityPtr->GetAllUserAccounts();
-								if (!UserAccounts.IsEmpty())
-								{
-									BeaconConnection->PlayerId = UserAccounts[0]->GetUserId();
-								}
-							}
-						}
+						// Send the player unique Id at login
+						BeaconConnection->PlayerId = GetFirstPlayerAccountId(World);
 					}
 
 #if NETCONNECTION_HAS_SETENCRYPTIONKEY
@@ -391,10 +420,10 @@ void AOnlineBeaconClient::NotifyControlMessage(UNetConnection* Connection, uint8
 			{
 				// build a URL
 				FURL URL(nullptr, TEXT(""), TRAVEL_Absolute);
-				FString AuthTicket = GetAuthTicket(*Connection->PlayerId);
-				if (!AuthTicket.IsEmpty())
+				const FString& LoginOptions = GetLoginOptions(Connection->PlayerId);
+				if (!LoginOptions.IsEmpty())
 				{
-					URL.AddOption(*FString::Printf(TEXT("AuthTicket=%s"), *AuthTicket));
+					URL.AddOption(*LoginOptions);
 				}
 				FString URLString = URL.ToString();
 
@@ -518,7 +547,15 @@ void AOnlineBeaconClient::NotifyControlMessage(UNetConnection* Connection, uint8
 	}	
 }
 
+PRAGMA_DISABLE_DEPRECATION_WARNINGS
 FString AOnlineBeaconClient::GetAuthTicket(const FUniqueNetIdRepl& PlayerId)
+PRAGMA_ENABLE_DEPRECATION_WARNINGS
+{
+	return GetAuthTicketInternal(PlayerId);
+}
+
+
+FString AOnlineBeaconClient::GetAuthTicketInternal(const FUniqueNetIdRepl& PlayerId)
 {
 	FString AuthTicket;
 
diff --git a/Engine/Plugins/Online/OnlineSubsystemUtils/Source/OnlineSubsystemUtils/Public/OnlineBeaconClient.h b/Engine/Plugins/Online/OnlineSubsystemUtils/Source/OnlineSubsystemUtils/Public/OnlineBeaconClient.h
--- a/Engine/Plugins/Online/OnlineSubsystemUtils/Source/OnlineSubsystemUtils/Public/OnlineBeaconClient.h
+++ b/Engine/Plugins/Online/OnlineSubsystemUtils/Source/OnlineSubsystemUtils/Public/OnlineBeaconClient.h
@@ -115,14 +115,24 @@ class ONLINESUBSYSTEMUTILS_API AOnlineBeaconClient : public AOnlineBeacon
 	const FUniqueNetIdRepl& GetUniqueId() const;
 
 	/**
-	* Get the authentication ticket to send when the server requires auth.
-	* 
-	* @param PlayerId unique id of the user on this connection
-	* 
-	* @return AuthTicket value to send to server
-	*/
+	 * Get the authentication ticket to send when the server requires auth.
+	 * 
+	 * @param PlayerId unique id of the user on this connection
+	 * 
+	 * @return AuthTicket value to send to server
+	 */
+	UE_DEPRECATED(5.3, "GetAuthTicket is deprecated. Override the login options by overriding GetLoginOptions and specifying AuthTicket=%s with the resulting string")
 	virtual FString GetAuthTicket(const FUniqueNetIdRepl& PlayerId);
 
+	/**
+	 * Get the login options to send the server during login.
+	 * Typically this is expected to be in the format that UGameplayStatics::ParseOption can parse, i.e. Key1=Value1?Key2=Value2?Key3=Value3
+	 *
+	 * @param PlayerId unique id of the user on this connection
+	 * @return Login options value to send to server
+	 */
+	virtual FString GetLoginOptions(const FUniqueNetIdRepl& PlayerId);
+
 	/**
 	 * Get the connection state
 	 * Higher level than the net connection because of the handshaking of the actors
@@ -179,6 +189,9 @@ protected:
 
 private:
 
+	/** Get the auth ticket from OnlineSubsystem for the specified player id */
+	FString GetAuthTicketInternal(const FUniqueNetIdRepl& PlayerId);
+
 	/** Encryption data used when connecting to servers. */
 	FEncryptionData EncryptionData;
 
