Update beacon host to provide all login options to StartVerifyAuthentication

diff --git a/Engine/Plugins/Online/OnlineSubsystemUtils/Source/OnlineSubsystemUtils/Private/OnlineBeaconHost.cpp b/Engine/Plugins/Online/OnlineSubsystemUtils/Source/OnlineSubsystemUtils/Private/OnlineBeaconHost.cpp
--- a/Engine/Plugins/Online/OnlineSubsystemUtils/Source/OnlineSubsystemUtils/Private/OnlineBeaconHost.cpp
+++ b/Engine/Plugins/Online/OnlineSubsystemUtils/Source/OnlineSubsystemUtils/Private/OnlineBeaconHost.cpp
@@ -314,14 +314,14 @@ bool AOnlineBeaconHost::HandleControlMessage(UNetConnection* Connection, uint8 M
 			// Try to kick off verification for this player.
 			const FString AuthTicket = UGameplayStatics::ParseOption(OptionsURL, TEXT("AuthTicket"));
 
+			bool bStartedAuth = false;
+
 			// Try to start deprecated auth method.
 			PRAGMA_DISABLE_DEPRECATION_WARNINGS
-			const bool bStartedAuthDeprecated = StartVerifyAuthentication(*UniqueIdRepl, AuthTicket);
+			bStartedAuth = StartVerifyAuthentication(*UniqueIdRepl, AuthTicket);
 			PRAGMA_ENABLE_DEPRECATION_WARNINGS
 
-			// Don't start new auth method if deprecated auth is active.
-			bool bStartedNewAuth = false;
-			if (!bStartedAuthDeprecated)
+			if (!bStartedAuth)
 			{
 				// Create completion delegate.
 				FOnAuthenticationVerificationCompleteDelegate OnAuthComplete = FOnAuthenticationVerificationCompleteDelegate::CreateWeakLambda(this, [this, WeakConnection = TWeakObjectPtr<UNetConnection>(Connection)](const FOnlineError& OnlineError)
@@ -329,10 +329,19 @@ bool AOnlineBeaconHost::HandleControlMessage(UNetConnection* Connection, uint8 M
 					OnAuthenticationVerificationComplete(WeakConnection.Get(), OnlineError);
 				});
 
-				bStartedNewAuth = StartVerifyAuthentication(*UniqueIdRepl, AuthTicket, OnAuthComplete);
+				// Try to start deprecated auth method.
+				PRAGMA_DISABLE_DEPRECATION_WARNINGS
+				bStartedAuth = StartVerifyAuthentication(*UniqueIdRepl, AuthTicket, OnAuthComplete);
+				PRAGMA_ENABLE_DEPRECATION_WARNINGS
+
+				// Don't start new auth method if deprecated auth is active.
+				if (!bStartedAuth)
+				{
+					bStartedAuth = StartVerifyAuthentication(*UniqueIdRepl, OptionsURL, AuthTicket, OnAuthComplete);
+				}
 			}
 
-			if (!bStartedAuthDeprecated && !bStartedNewAuth)
+			if (!bStartedAuth)
 			{
 				static const FText ErrorTxt = NSLOCTEXT("NetworkErrors", "BeaconLoginInvalidAuthHandlerError", "Login Failure. Unable to process authentication.");
 				SendFailurePacket(Connection, ENetCloseResult::BeaconLoginInvalidAuthHandlerError, ErrorTxt);
@@ -628,12 +637,16 @@ void AOnlineBeaconHost::RemoveClientActor(AOnlineBeaconClient* ClientActor)
 	}
 }
 
+PRAGMA_DISABLE_DEPRECATION_WARNINGS
 bool AOnlineBeaconHost::StartVerifyAuthentication(const FUniqueNetId& PlayerId, const FString& AuthenticationToken)
+PRAGMA_ENABLE_DEPRECATION_WARNINGS
 {
 	return false;
 }
 
+PRAGMA_DISABLE_DEPRECATION_WARNINGS
 void AOnlineBeaconHost::OnAuthenticationVerificationComplete(const class FUniqueNetId& PlayerId, const FOnlineError& Error)
+PRAGMA_ENABLE_DEPRECATION_WARNINGS
 {
 	UNetConnection* Connection = nullptr;
 	FConnectionState* ConnState = nullptr;
@@ -643,7 +656,14 @@ void AOnlineBeaconHost::OnAuthenticationVerificationComplete(const class FUnique
 	}
 }
 
+PRAGMA_DISABLE_DEPRECATION_WARNINGS
 bool AOnlineBeaconHost::StartVerifyAuthentication(const FUniqueNetId& PlayerId, const FString& AuthenticationToken, const FOnAuthenticationVerificationCompleteDelegate& OnComplete)
+PRAGMA_ENABLE_DEPRECATION_WARNINGS
+{
+	return false;
+}
+
+bool AOnlineBeaconHost::StartVerifyAuthentication(const FUniqueNetId& PlayerId, const FString& LoginOptions, const FString& AuthenticationToken, const FOnAuthenticationVerificationCompleteDelegate& OnComplete)
 {
 	return false;
 }
diff --git a/Engine/Plugins/Online/OnlineSubsystemUtils/Source/OnlineSubsystemUtils/Public/OnlineBeaconHost.h b/Engine/Plugins/Online/OnlineSubsystemUtils/Source/OnlineSubsystemUtils/Public/OnlineBeaconHost.h
--- a/Engine/Plugins/Online/OnlineSubsystemUtils/Source/OnlineSubsystemUtils/Public/OnlineBeaconHost.h
+++ b/Engine/Plugins/Online/OnlineSubsystemUtils/Source/OnlineSubsystemUtils/Public/OnlineBeaconHost.h
@@ -156,8 +156,20 @@ protected:
 	 * @param PlayerId net id of player to authenticate.
 	 * @param AuthenticationToken token to use for verification.
 	 */
+	UE_DEPRECATED(5.3, "This version of the StartVerifyAuthentication is deprecated. Please use the new StartVerifyAuthentication method instead.")
 	virtual bool StartVerifyAuthentication(const FUniqueNetId& PlayerId, const FString& AuthenticationToken, const FOnAuthenticationVerificationCompleteDelegate& OnComplete);
 
+	/**
+	 * Start verifying an authentication token for a connection.
+	 * OnAuthenticationVerificationComplete must be called to complete authentication verification.
+	 *
+	 * @param PlayerId net id of player to authenticate.
+	 * @Param LoginOptions all options passed as part of the Login request.
+	 * @param AuthenticationToken token to use for verification.
+	 * @Param OnComplete delegate to call once the request for authentication has completed
+	 */
+	virtual bool StartVerifyAuthentication(const FUniqueNetId& PlayerId, const FString& LoginOptions, const FString& AuthenticationToken, const FOnAuthenticationVerificationCompleteDelegate& OnComplete);
+
 private:
 	/**
 	 * Event which must be signaled to complete an authentication verification request.
