Add VerifyAuthForBeaconType method to OnlineBeaconHost to allow dervied classes to approve authentication once the beacon type is known.

diff --git a/Engine/Plugins/Online/OnlineSubsystemUtils/Source/OnlineSubsystemUtils/Private/OnlineBeaconHost.cpp b/Engine/Plugins/Online/OnlineSubsystemUtils/Source/OnlineSubsystemUtils/Private/OnlineBeaconHost.cpp
--- a/Engine/Plugins/Online/OnlineSubsystemUtils/Source/OnlineSubsystemUtils/Private/OnlineBeaconHost.cpp
+++ b/Engine/Plugins/Online/OnlineSubsystemUtils/Source/OnlineSubsystemUtils/Private/OnlineBeaconHost.cpp
@@ -407,6 +407,18 @@ bool AOnlineBeaconHost::HandleControlMessage(UNetConnection* Connection, uint8 M
 
 					return false;
 				}
+
+				if (!VerifyJoinForBeaconType(*UniqueId, BeaconType))
+				{
+					static const FText AuthErrorText = NSLOCTEXT("NetworkErrors", "BeaconAuthError", "Unable to authenticate for beacon. Verifying auth for beacon type {0} failed for connection owned by {1}");
+
+					SendFailurePacket(Connection, ENetCloseResult::BeaconAuthError,
+										FText::Format(AuthErrorText, FText::FromString(BeaconType),
+														FText::FromString(Connection->PlayerId.ToDebugString())));
+
+					return false;
+				}
+
 				UE_LOG(LogBeacon, Log, TEXT("%s: Beacon Join %s %s"), *GetDebugName(Connection), *BeaconType, *UniqueId.ToDebugString());
 			}
 			else
@@ -668,6 +680,11 @@ bool AOnlineBeaconHost::StartVerifyAuthentication(const FUniqueNetId& PlayerId,
 	return false;
 }
 
+bool AOnlineBeaconHost::VerifyJoinForBeaconType(const FUniqueNetId& PlayerId, const FString& BeaconType)
+{
+	return true;
+}
+
 void AOnlineBeaconHost::OnAuthenticationVerificationComplete(UNetConnection* Connection, const FOnlineError& Error)
 {
 	if (FConnectionState* ConnState = ConnectionState.Find(Connection))
diff --git a/Engine/Plugins/Online/OnlineSubsystemUtils/Source/OnlineSubsystemUtils/Public/OnlineBeaconHost.h b/Engine/Plugins/Online/OnlineSubsystemUtils/Source/OnlineSubsystemUtils/Public/OnlineBeaconHost.h
--- a/Engine/Plugins/Online/OnlineSubsystemUtils/Source/OnlineSubsystemUtils/Public/OnlineBeaconHost.h
+++ b/Engine/Plugins/Online/OnlineSubsystemUtils/Source/OnlineSubsystemUtils/Public/OnlineBeaconHost.h
@@ -170,6 +170,14 @@ protected:
 	 */
 	virtual bool StartVerifyAuthentication(const FUniqueNetId& PlayerId, const FString& LoginOptions, const FString& AuthenticationToken, const FOnAuthenticationVerificationCompleteDelegate& OnComplete);
 
+	/**
+	 * Verify user authentication once the beacon type is known.
+	 *
+	 * @param PlayerId net id of player to authenticate.
+	 * @Param BeaconType type of beacon to verify the authentication for.
+	 */
+	virtual bool VerifyJoinForBeaconType(const FUniqueNetId& PlayerId, const FString& BeaconType);
+
 private:
 	/**
 	 * Event which must be signaled to complete an authentication verification request.
