Adding ability to specify a port when using "net recreatesocket"

diff --git a/Engine/Plugins/Online/OnlineSubsystemUtils/Source/OnlineSubsystemUtils/Classes/IpNetDriver.h b/Engine/Plugins/Online/OnlineSubsystemUtils/Source/OnlineSubsystemUtils/Classes/IpNetDriver.h
--- a/Engine/Plugins/Online/OnlineSubsystemUtils/Source/OnlineSubsystemUtils/Classes/IpNetDriver.h
+++ b/Engine/Plugins/Online/OnlineSubsystemUtils/Source/OnlineSubsystemUtils/Classes/IpNetDriver.h
@@ -161,9 +161,11 @@ public:
 	 *
 	 * Also usable to test the 'restart handshake' feature, using the 'net RecreateSocket' command.
 	 *
+	 * @param OverridePort		optional port number to use with the new socket, otherwise falls back to GetClientPort()
+	 * 
 	 * @return	Returns the action taken (e.g. successfully kicked off recreation, or whether this failed or is already in progress etc.)
 	 */
-	UE::Net::ERecreateSocketResult RecreateSocket();
+	UE::Net::ERecreateSocketResult RecreateSocket(int32 OverridePort = INDEX_NONE);
 
 	/** Returns the last time socket recreation was kicked off or completed */
 	double GetLastRecreateSocketTime() const
diff --git a/Engine/Plugins/Online/OnlineSubsystemUtils/Source/OnlineSubsystemUtils/Private/IpNetDriver.cpp b/Engine/Plugins/Online/OnlineSubsystemUtils/Source/OnlineSubsystemUtils/Private/IpNetDriver.cpp
--- a/Engine/Plugins/Online/OnlineSubsystemUtils/Source/OnlineSubsystemUtils/Private/IpNetDriver.cpp
+++ b/Engine/Plugins/Online/OnlineSubsystemUtils/Source/OnlineSubsystemUtils/Private/IpNetDriver.cpp
@@ -1663,7 +1663,15 @@ bool UIpNetDriver::HandleRecreateSocketCommand(const TCHAR* Cmd, FOutputDevice&
 {
 	using namespace UE::Net;
 
-	ERecreateSocketResult Result = RecreateSocket();
+	FString PortStr;
+	int32 OverridePort = INDEX_NONE;
+
+	if (FParse::Token(Cmd, PortStr, false))
+	{
+		OverridePort = FCString::Atoi(*PortStr);
+	}
+
+	ERecreateSocketResult Result = RecreateSocket(OverridePort);
 
 	if (Result == ERecreateSocketResult::BeganRecreate)
 	{
@@ -1878,7 +1886,7 @@ void UIpNetDriver::TickNewIPTracking(float DeltaTime)
 	}
 }
 
-UE::Net::ERecreateSocketResult UIpNetDriver::RecreateSocket()
+UE::Net::ERecreateSocketResult UIpNetDriver::RecreateSocket(int32 OverridePort)
 {
 	using namespace UE::Net;
 	using namespace UE::Net::Private;
@@ -1894,8 +1902,10 @@ UE::Net::ERecreateSocketResult UIpNetDriver::RecreateSocket()
 
 		if (CurState == ESocketState::Ready)
 		{
+			const int32 BindPort = OverridePort != INDEX_NONE ? OverridePort : GetClientPort();
+
 			FString Error;
-			FUniqueSocket NewSocket = CreateAndBindSocket(LocalAddr.ToSharedRef(), GetClientPort(), false, ClientDesiredSocketReceiveBufferBytes,
+			FUniqueSocket NewSocket = CreateAndBindSocket(LocalAddr.ToSharedRef(), BindPort, false, ClientDesiredSocketReceiveBufferBytes,
 															ClientDesiredSocketSendBufferBytes, Error);
 
 			if (NewSocket.IsValid())
