Online PIE settings: Support command line args
This is mostly for launchit with scripts that needs credidentials in PIE
#rnx

diff --git a/Engine/Plugins/Online/OnlineSubsystemUtils/Source/OnlineSubsystemUtils/Private/OnlineServicesEngineUtilsImpl.cpp b/Engine/Plugins/Online/OnlineSubsystemUtils/Source/OnlineSubsystemUtils/Private/OnlineServicesEngineUtilsImpl.cpp
--- a/Engine/Plugins/Online/OnlineSubsystemUtils/Source/OnlineSubsystemUtils/Private/OnlineServicesEngineUtilsImpl.cpp
+++ b/Engine/Plugins/Online/OnlineSubsystemUtils/Source/OnlineSubsystemUtils/Private/OnlineServicesEngineUtilsImpl.cpp
@@ -10,6 +10,11 @@
 #include "Online/OnlineServicesDelegates.h"
 #include "Online/OnlineServicesRegistry.h"
 
+#if WITH_EDITOR
+#include "Misc/Parse.h"
+#endif
+
+
 namespace UE::Online
 {
 
@@ -85,11 +90,28 @@ void FOnlineServicesEngineUtils::SetEngineExternalUIBinding(const FOnExternalUIC
 }
 
 #if WITH_EDITOR
+
+namespace FOnlineServicesEngineUtilsHelper
+{
+	bool HasAuthCommandLine(UE::Online::FAuthLogin::Params& LoginParameters)
+    {
+    	check(LoginParameters.CredentialsToken.IsType<FString>());
+    	return FParse::Value(FCommandLine::Get(), TEXT("AUTH_LOGIN="), LoginParameters.CredentialsId ) &&
+    		   FParse::Value(FCommandLine::Get(), TEXT("AUTH_PASSWORD="), LoginParameters.CredentialsToken.Get<FString>()) &&
+    		   FParse::Value(FCommandLine::Get(), TEXT("AUTH_TYPE="), LoginParameters.CredentialsType);
+    }
+	bool HasAuthCommandLine()
+	{
+		UE::Online::FAuthLogin::Params LoginParameters;
+		return HasAuthCommandLine(LoginParameters);
+	}
+}
+
 bool FOnlineServicesEngineUtils::SupportsOnlinePIE() const
 {
 	check(UObjectInitialized());
 	const UOnlinePIESettings* OnlinePIESettings = GetDefault<UOnlinePIESettings>();
-	if (OnlinePIESettings->bOnlinePIEEnabled && GetNumPIELogins() > 0)
+	if ((OnlinePIESettings->bOnlinePIEEnabled && GetNumPIELogins() > 0) || FOnlineServicesEngineUtilsHelper::HasAuthCommandLine())
 	{
 		// If we can't get the auth then things are either not configured right or disabled
 		EOnlineServices OnlineServicesType = EOnlineServices::Epic; // TODO:  Need Default support UE::Online::EOnlineServices::Default
@@ -116,8 +138,12 @@ void FOnlineServicesEngineUtils::SetShouldTryOnlinePIE(bool bShouldTry)
 bool FOnlineServicesEngineUtils::IsOnlinePIEEnabled() const
 {
 	check(UObjectInitialized());
+	if (!bShouldTryOnlinePIE)
+	{
+		return false;
+	}
 	const UOnlinePIESettings* OnlinePIESettings = GetDefault<UOnlinePIESettings>();
-	return bShouldTryOnlinePIE && OnlinePIESettings->bOnlinePIEEnabled;
+	return OnlinePIESettings->bOnlinePIEEnabled || FOnlineServicesEngineUtilsHelper::HasAuthCommandLine();
 }
 
 int32 FOnlineServicesEngineUtils::GetNumPIELogins() const
@@ -133,6 +159,10 @@ int32 FOnlineServicesEngineUtils::GetNumPIELogins() const
 			NumValidLogins++;
 		}
 	}
+	if (NumValidLogins == 0 && FOnlineServicesEngineUtilsHelper::HasAuthCommandLine())
+	{
+		NumValidLogins++;
+	}
 	
 	return NumValidLogins;
 }
@@ -149,6 +179,14 @@ void FOnlineServicesEngineUtils::GetPIELogins(TArray<FAuthLogin::Params>& Logins
 		LoginParameters.CredentialsType = *PIELoginCredentials.Type;
 		Logins.Emplace(MoveTemp(LoginParameters));
 	}
+	if (Logins.IsEmpty())
+	{
+		UE::Online::FAuthLogin::Params LoginParameters;
+		if (FOnlineServicesEngineUtilsHelper::HasAuthCommandLine(LoginParameters))
+		{
+			Logins.Emplace(MoveTemp(LoginParameters));
+		}
+	}
 }
 
 #endif // WITH_EDITOR
diff --git a/Engine/Plugins/Online/OnlineSubsystemUtils/Source/OnlineSubsystemUtils/Private/OnlineSubsystemUtilsModule.cpp b/Engine/Plugins/Online/OnlineSubsystemUtils/Source/OnlineSubsystemUtils/Private/OnlineSubsystemUtilsModule.cpp
--- a/Engine/Plugins/Online/OnlineSubsystemUtils/Source/OnlineSubsystemUtils/Private/OnlineSubsystemUtilsModule.cpp
+++ b/Engine/Plugins/Online/OnlineSubsystemUtils/Source/OnlineSubsystemUtils/Private/OnlineSubsystemUtilsModule.cpp
@@ -7,6 +7,10 @@
 #include "OnlineDelegates.h"
 #include "OnlinePIESettings.h"
 
+#if WITH_EDITOR
+#include "Misc/Parse.h"
+#endif
+
 IMPLEMENT_MODULE(FOnlineSubsystemUtilsModule, OnlineSubsystemUtils);
 
 /**
@@ -103,11 +107,25 @@ public:
 	}
 
 #if WITH_EDITOR
+
+private:
+	bool HasAuthCommandLine(FOnlineAccountCredentials& LoginParameters) const
+	{
+		return FParse::Value(FCommandLine::Get(), TEXT("AUTH_LOGIN="), LoginParameters.Id ) &&
+			   FParse::Value(FCommandLine::Get(), TEXT("AUTH_PASSWORD="), LoginParameters.Token) &&
+			   FParse::Value(FCommandLine::Get(), TEXT("AUTH_TYPE="), LoginParameters.Type);
+	}
+	bool HasAuthCommandLine() const
+	{
+		FOnlineAccountCredentials LoginParameters;
+		return HasAuthCommandLine(LoginParameters);
+	}
+public:	
 	virtual bool SupportsOnlinePIE() const override
 	{
 		check(UObjectInitialized());
 		const UOnlinePIESettings* OnlinePIESettings = GetDefault<UOnlinePIESettings>();
-		if (OnlinePIESettings->bOnlinePIEEnabled && GetNumPIELogins() > 0)
+		if ((OnlinePIESettings->bOnlinePIEEnabled && GetNumPIELogins() > 0) || HasAuthCommandLine())
 		{
 			// If we can't get the identity interface then things are either not configured right or disabled
 			IOnlineIdentityPtr IdentityInt = Online::GetIdentityInterface();
@@ -131,8 +149,12 @@ public:
 	virtual bool IsOnlinePIEEnabled() const override
 	{
 		check(UObjectInitialized());
+		if (!bShouldTryOnlinePIE)
+		{
+			return false;
+		}
 		const UOnlinePIESettings* OnlinePIESettings = GetDefault<UOnlinePIESettings>();
-		return bShouldTryOnlinePIE && OnlinePIESettings->bOnlinePIEEnabled;
+		return OnlinePIESettings->bOnlinePIEEnabled || HasAuthCommandLine();
 	}
 
 	virtual int32 GetNumPIELogins() const override
@@ -148,6 +170,10 @@ public:
 				NumValidLogins++;
 			}
 		}
+		if (NumValidLogins == 0 && HasAuthCommandLine())
+		{
+			NumValidLogins++;
+		}
 	
 		return NumValidLogins;
 	}
@@ -158,17 +184,23 @@ public:
 		const UOnlinePIESettings* OnlinePIESettings = GetDefault<UOnlinePIESettings>();
 		if (OnlinePIESettings->Logins.Num() > 0)
 		{
-			FOnlineAccountCredentials TempLogin;
-
 			Logins.Empty(OnlinePIESettings->Logins.Num());
 			for (const FPIELoginSettingsInternal& Login : OnlinePIESettings->Logins)
 			{
 				if (Login.IsValid())
 				{
-					new (Logins)FOnlineAccountCredentials(Login.Type, Login.Id, Login.Token);
+					Logins.Emplace(Login.Type, Login.Id, Login.Token);
 				}
 			}
 		}
+		if (Logins.IsEmpty())
+		{
+			FOnlineAccountCredentials LoginParameters;
+			if (HasAuthCommandLine(LoginParameters))
+			{
+				Logins.Emplace(MoveTemp(LoginParameters));
+			}
+		}
 	}
 
 #endif // WITH_EDITOR
