[VK]

Adding functionality to assign players to a team they can actually fit to when respect parties is configured.

diff --git a/Engine/Plugins/Online/OnlineSubsystemUtils/Source/OnlineSubsystemUtils/Private/PartyBeaconHost.cpp b/Engine/Plugins/Online/OnlineSubsystemUtils/Source/OnlineSubsystemUtils/Private/PartyBeaconHost.cpp
--- a/Engine/Plugins/Online/OnlineSubsystemUtils/Source/OnlineSubsystemUtils/Private/PartyBeaconHost.cpp
+++ b/Engine/Plugins/Online/OnlineSubsystemUtils/Source/OnlineSubsystemUtils/Private/PartyBeaconHost.cpp
@@ -80,6 +80,11 @@ bool APartyBeaconHost::ReconfigureTeamAndPlayerCount(int32 InNumTeams, int32 InN
 	return bSuccess;
 }
 
+void APartyBeaconHost::SetCompetitiveIntegrity(bool bNewCompetitiveIntegrity)
+{
+	State->SetCompetitiveIntegrity(bNewCompetitiveIntegrity);
+}
+
 void APartyBeaconHost::SetTeamAssignmentMethod(FName NewAssignmentMethod)
 {
 	if (State)
@@ -629,12 +634,15 @@ EPartyReservationResult::Type APartyBeaconHost::AddPartyReservation(const FParty
 			}	
 			if (TeamIdx != -1)
 			{	
-				// Will the part members not currently on the team fit in it, and fit within the server
-				const int32 NumTeamMembers = GetNumPlayersOnTeam(TeamIdx);
-				const int32 NumAvailableSlotsOnTeam = FMath::Max<int32>(0, (GetMaxPlayersPerTeam() - NumTeamMembers) + NumTeamPlayersWithExistingReservation);
-				if ((NumAvailableSlotsOnTeam < ReservationRequest.PartyMembers.Num()) || (State->GetRemainingReservations() < (ReservationRequest.PartyMembers.Num() - NumTeamPlayersWithExistingReservation)))
+				if (ShouldRespectCompetitiveIntegrity())
 				{
-					bContainsIncompatibleExistingMembers = true;
+					// Will the party members not currently on the team fit in it, and fit within the server
+					const int32 NumTeamMembers = GetNumPlayersOnTeam(TeamIdx);
+					const int32 NumAvailableSlotsOnTeam = FMath::Max<int32>(0, (GetMaxPlayersPerTeam() - NumTeamMembers) + NumTeamPlayersWithExistingReservation);
+					if ((NumAvailableSlotsOnTeam < ReservationRequest.PartyMembers.Num()) || (State->GetRemainingReservations() < (ReservationRequest.PartyMembers.Num() - NumTeamPlayersWithExistingReservation)))
+					{
+						bContainsIncompatibleExistingMembers = true;
+					}
 				}
 			}
 			// if this party reservation included players who already had reservations, but is not incompatible based on teams then remove those players previous reservations
@@ -888,7 +896,7 @@ EPartyReservationResult::Type APartyBeaconHost::UpdatePartyReservation(const FPa
 				if ((State->GetRemainingReservations() - NewPlayers.Num()) >= 0)
 				{
 					// Validate that adding the new party members to this reservation entry still fits within the team size
-					if ((NewPlayers.Num() - NumPlayersWithExistingReservation) <= NumAvailableSlotsOnTeam)
+					if (!ShouldRespectCompetitiveIntegrity() || (NewPlayers.Num() - NumPlayersWithExistingReservation) <= NumAvailableSlotsOnTeam)
 					{
 						bool bPlayerRemovedFromReservation = false;
 						if (NewPlayers.Num() > 0)
diff --git a/Engine/Plugins/Online/OnlineSubsystemUtils/Source/OnlineSubsystemUtils/Private/PartyBeaconState.cpp b/Engine/Plugins/Online/OnlineSubsystemUtils/Source/OnlineSubsystemUtils/Private/PartyBeaconState.cpp
--- a/Engine/Plugins/Online/OnlineSubsystemUtils/Source/OnlineSubsystemUtils/Private/PartyBeaconState.cpp
+++ b/Engine/Plugins/Online/OnlineSubsystemUtils/Source/OnlineSubsystemUtils/Private/PartyBeaconState.cpp
@@ -120,7 +120,8 @@ UPartyBeaconState::UPartyBeaconState(const FObjectInitializer& ObjectInitializer
 	TeamAssignmentMethod(ETeamAssignmentMethod::Smallest),
 	ReservedHostTeamNum(0),
 	ForceTeamNum(0),
-	bRestrictCrossConsole(false)
+	bRestrictCrossConsole(false),
+	bRespectCompetitiveIntegrity(true)
 {
 }
 
@@ -221,6 +222,11 @@ bool UPartyBeaconState::ReconfigureTeamAndPlayerCount(int32 InNumTeams, int32 In
 	return bSuccess;
 }
 
+void UPartyBeaconState::SetCompetitiveIntegrity(bool bNewCompetitiveIntegrity)
+{
+	bRespectCompetitiveIntegrity = bNewCompetitiveIntegrity;
+}
+
 int32 UPartyBeaconState::GetMaxAvailableTeamSize() const
 {
 	int32 MaxFreeSlots = 0;
diff --git a/Engine/Plugins/Online/OnlineSubsystemUtils/Source/OnlineSubsystemUtils/Public/PartyBeaconHost.h b/Engine/Plugins/Online/OnlineSubsystemUtils/Source/OnlineSubsystemUtils/Public/PartyBeaconHost.h
--- a/Engine/Plugins/Online/OnlineSubsystemUtils/Source/OnlineSubsystemUtils/Public/PartyBeaconHost.h
+++ b/Engine/Plugins/Online/OnlineSubsystemUtils/Source/OnlineSubsystemUtils/Public/PartyBeaconHost.h
@@ -97,6 +97,13 @@ class ONLINESUBSYSTEMUTILS_API APartyBeaconHost : public AOnlineBeaconHostObject
 	 */
 	virtual bool ReconfigureTeamAndPlayerCount(int32 InNumTeams, int32 InNumPlayersPerTeam, int32 InNumReservations);
 
+	/**
+	 * Competitive integrity represents how "fair" we need to keep the game. So no friends on other teams, no cross team chat, etc.
+	 *
+	 * @param bNewCompetitiveIntegrity The new value for Competitive Integrity
+	 */
+	virtual void SetCompetitiveIntegrity(bool bNewCompetitiveIntegrity);
+
 	/**
 	 * Define the method for assignment new reservations to teams
 	 * 
@@ -180,6 +187,13 @@ class ONLINESUBSYSTEMUTILS_API APartyBeaconHost : public AOnlineBeaconHostObject
 	 */
 	virtual int32 GetMaxPlayersPerTeam() const { return State->GetMaxPlayersPerTeam(); }
 
+	/**
+	 * Competitive integrity represents how "fair" we need to keep the game. So no friends on other teams, no cross team chat, etc.
+	 *
+	 * @return Whether or not we should respect competitive integrity
+	 */
+	virtual bool ShouldRespectCompetitiveIntegrity() const { return State->ShouldRespectCompetitiveIntegrity(); }
+
 	/**
 	 * Determine the maximum team size that can be accommodated based
 	 * on the current reservation slots occupied.
diff --git a/Engine/Plugins/Online/OnlineSubsystemUtils/Source/OnlineSubsystemUtils/Public/PartyBeaconState.h b/Engine/Plugins/Online/OnlineSubsystemUtils/Source/OnlineSubsystemUtils/Public/PartyBeaconState.h
--- a/Engine/Plugins/Online/OnlineSubsystemUtils/Source/OnlineSubsystemUtils/Public/PartyBeaconState.h
+++ b/Engine/Plugins/Online/OnlineSubsystemUtils/Source/OnlineSubsystemUtils/Public/PartyBeaconState.h
@@ -262,6 +262,13 @@ class ONLINESUBSYSTEMUTILS_API UPartyBeaconState : public UObject
 	 */
 	virtual bool ReconfigureTeamAndPlayerCount(int32 InNumTeams, int32 InNumPlayersPerTeam, int32 InNumReservations);
 
+	/**
+	 * Competitive integrity represents how "fair" we need to keep the game. So no friends on other teams, no cross team chat, etc.
+	 *
+	 * @param bNewCompetitiveIntegrity New value for Competitive integrity
+	 */
+	virtual void SetCompetitiveIntegrity(bool bNewCompetitiveIntegrity);
+
 	/**
 	 * Define the method for assignment new reservations to teams
 	 * 
@@ -413,6 +420,13 @@ class ONLINESUBSYSTEMUTILS_API UPartyBeaconState : public UObject
 	 */
 	virtual int32 GetMaxPlayersPerTeam() const { return NumPlayersPerTeam; }
 
+	/**
+	 * Competitive integrity represents how "fair" we need to keep the game. So no friends on other teams, no cross team chat, etc.
+	 *
+	 * @return Whether we should respect competitive integrity
+	 */
+	virtual bool ShouldRespectCompetitiveIntegrity() const { return bRespectCompetitiveIntegrity; }
+
 	/**
 	 * Determine the maximum team size that can be accommodated based
 	 * on the current reservation slots occupied.
@@ -578,6 +592,9 @@ protected:
 	/** Process requests from clients to remove players from beacon */
 	UPROPERTY(Transient)
 	bool bEnableRemovalRequests;
+	/** Competitive integrity represents how "fair" we need to keep the game. So no friends on other teams, no cross team chat, etc. */
+	UPROPERTY(Transient)
+	bool bRespectCompetitiveIntegrity;
 
 	/** Current reservations in the system */
 	UPROPERTY(Transient)
