Allow party and spectator beacon validation strings to be optional.

diff --git a/Engine/Plugins/Online/OnlineSubsystemUtils/Source/OnlineSubsystemUtils/Private/PartyBeaconHost.cpp b/Engine/Plugins/Online/OnlineSubsystemUtils/Source/OnlineSubsystemUtils/Private/PartyBeaconHost.cpp
--- a/Engine/Plugins/Online/OnlineSubsystemUtils/Source/OnlineSubsystemUtils/Private/PartyBeaconHost.cpp
+++ b/Engine/Plugins/Online/OnlineSubsystemUtils/Source/OnlineSubsystemUtils/Private/PartyBeaconHost.cpp
@@ -14,7 +14,8 @@
 APartyBeaconHost::APartyBeaconHost(const FObjectInitializer& ObjectInitializer) :
 	Super(ObjectInitializer),
 	State(NULL),
-	bLogoutOnSessionTimeout(true)
+	bLogoutOnSessionTimeout(true),
+	bIsValidationStrRequired(true)
 {
 	ClientBeaconActorClass = APartyBeaconClient::StaticClass();
 	BeaconTypeName = ClientBeaconActorClass->GetName();
@@ -508,7 +509,7 @@ EPartyReservationResult::Type APartyBeaconHost::AddPartyReservation(const FParty
 		return EPartyReservationResult::ReservationDenied;
 	}
 
-	if (ReservationRequest.IsValid())
+	if (ReservationRequest.IsValid(bIsValidationStrRequired))
 	{
 		TArray<FPartyReservation>& Reservations = State->GetReservations();
 		const int32 ExistingReservationIdx = State->GetExistingReservation(ReservationRequest.PartyLeader);
@@ -754,7 +755,7 @@ EPartyReservationResult::Type APartyBeaconHost::UpdatePartyReservation(const FPa
 		return EPartyReservationResult::ReservationDenied;
 	}
 
-	if (ReservationUpdateRequest.IsValid() || bIsRemovingMembers)
+	if (ReservationUpdateRequest.IsValid(bIsValidationStrRequired) || bIsRemovingMembers)
 	{
 		if (bIsRemovingMembers)
 		{
diff --git a/Engine/Plugins/Online/OnlineSubsystemUtils/Source/OnlineSubsystemUtils/Private/PartyBeaconState.cpp b/Engine/Plugins/Online/OnlineSubsystemUtils/Source/OnlineSubsystemUtils/Private/PartyBeaconState.cpp
--- a/Engine/Plugins/Online/OnlineSubsystemUtils/Source/OnlineSubsystemUtils/Private/PartyBeaconState.cpp
+++ b/Engine/Plugins/Online/OnlineSubsystemUtils/Source/OnlineSubsystemUtils/Private/PartyBeaconState.cpp
@@ -15,7 +15,7 @@ namespace ETeamAssignmentMethod
 	const FName Manual = FName(TEXT("Manual"));
 }
 
-bool FPartyReservation::IsValid() const
+bool FPartyReservation::IsValid(bool bIsValidationStrRequired) const
 {
 	bool bIsValid = false;
 	if (PartyLeader.IsValid() && PartyMembers.Num() >= 1)
@@ -35,7 +35,8 @@ bool FPartyReservation::IsValid() const
 				break;
 			}
 
-			if (PartyLeader == PlayerRes.UniqueId &&
+			if (bIsValidationStrRequired &&
+				PartyLeader == PlayerRes.UniqueId &&
 				PlayerRes.ValidationStr.IsEmpty())
 			{
 				bIsValid = false;
diff --git a/Engine/Plugins/Online/OnlineSubsystemUtils/Source/OnlineSubsystemUtils/Private/SpectatorBeaconHost.cpp b/Engine/Plugins/Online/OnlineSubsystemUtils/Source/OnlineSubsystemUtils/Private/SpectatorBeaconHost.cpp
--- a/Engine/Plugins/Online/OnlineSubsystemUtils/Source/OnlineSubsystemUtils/Private/SpectatorBeaconHost.cpp
+++ b/Engine/Plugins/Online/OnlineSubsystemUtils/Source/OnlineSubsystemUtils/Private/SpectatorBeaconHost.cpp
@@ -14,7 +14,8 @@
 ASpectatorBeaconHost::ASpectatorBeaconHost(const FObjectInitializer& ObjectInitializer) :
 	Super(ObjectInitializer),
 	State(NULL),
-	bLogoutOnSessionTimeout(true)
+	bLogoutOnSessionTimeout(true),
+	bIsValidationStrRequired(true)
 {
 	ClientBeaconActorClass = ASpectatorBeaconClient::StaticClass();
 	BeaconTypeName = ClientBeaconActorClass->GetName();
@@ -332,7 +333,7 @@ ESpectatorReservationResult::Type ASpectatorBeaconHost::AddSpectatorReservation(
 		return ESpectatorReservationResult::ReservationDenied;
 	}
 
-	if (ReservationRequest.IsValid())
+	if (ReservationRequest.IsValid(bIsValidationStrRequired))
 	{
 		TArray<FSpectatorReservation>& Reservations = State->GetReservations();
 		const int32 ExistingReservationIdx = State->GetExistingReservation(ReservationRequest.SpectatorId);
diff --git a/Engine/Plugins/Online/OnlineSubsystemUtils/Source/OnlineSubsystemUtils/Private/SpectatorBeaconState.cpp b/Engine/Plugins/Online/OnlineSubsystemUtils/Source/OnlineSubsystemUtils/Private/SpectatorBeaconState.cpp
--- a/Engine/Plugins/Online/OnlineSubsystemUtils/Source/OnlineSubsystemUtils/Private/SpectatorBeaconState.cpp
+++ b/Engine/Plugins/Online/OnlineSubsystemUtils/Source/OnlineSubsystemUtils/Private/SpectatorBeaconState.cpp
@@ -7,7 +7,7 @@
 
 DEFINE_LOG_CATEGORY(LogSpectatorBeacon);
 
-bool FSpectatorReservation::IsValid() const
+bool FSpectatorReservation::IsValid(bool bIsValidationStrRequired) const
 {
 	bool bIsValid = false;
 	if (SpectatorId.IsValid())
@@ -21,11 +21,11 @@ bool FSpectatorReservation::IsValid() const
 		{
 			bIsValid = false;
 		}
-		else if (SpectatorId == Spectator.UniqueId &&
-				Spectator.ValidationStr.IsEmpty())
+		else if (bIsValidationStrRequired &&
+			SpectatorId == Spectator.UniqueId &&
+			Spectator.ValidationStr.IsEmpty())
 		{
-				bIsValid = false;
-		
+			bIsValid = false;
 		}
 	}
 
diff --git a/Engine/Plugins/Online/OnlineSubsystemUtils/Source/OnlineSubsystemUtils/Public/PartyBeaconHost.h b/Engine/Plugins/Online/OnlineSubsystemUtils/Source/OnlineSubsystemUtils/Public/PartyBeaconHost.h
--- a/Engine/Plugins/Online/OnlineSubsystemUtils/Source/OnlineSubsystemUtils/Public/PartyBeaconHost.h
+++ b/Engine/Plugins/Online/OnlineSubsystemUtils/Source/OnlineSubsystemUtils/Public/PartyBeaconHost.h
@@ -387,6 +387,11 @@ class ONLINESUBSYSTEMUTILS_API APartyBeaconHost : public AOnlineBeaconHostObject
 	 */
 	virtual void DumpReservations() const;
 
+	/**
+	 * Do party members require validation strings.
+	 */
+	bool IsValidationStrRequired() const { return bIsValidationStrRequired; }
+
 protected:
 
 	/** State of the beacon */
@@ -409,6 +414,9 @@ protected:
 	/** Do the timeouts below cause a player to be removed from the reservation list */
 	UPROPERTY(Config)
 	bool bLogoutOnSessionTimeout;
+	/** Do party members require validation strings. */
+	UPROPERTY(Config)
+	bool bIsValidationStrRequired;
 	/** Seconds that can elapse before a reservation is removed due to player not being registered with the session */
 	UPROPERTY(Transient, Config)
 	float SessionTimeoutSecs;
diff --git a/Engine/Plugins/Online/OnlineSubsystemUtils/Source/OnlineSubsystemUtils/Public/PartyBeaconState.h b/Engine/Plugins/Online/OnlineSubsystemUtils/Source/OnlineSubsystemUtils/Public/PartyBeaconState.h
--- a/Engine/Plugins/Online/OnlineSubsystemUtils/Source/OnlineSubsystemUtils/Public/PartyBeaconState.h
+++ b/Engine/Plugins/Online/OnlineSubsystemUtils/Source/OnlineSubsystemUtils/Public/PartyBeaconState.h
@@ -197,7 +197,7 @@ struct ONLINESUBSYSTEMUTILS_API FPartyReservation
 	int32 RemoveAllPartyMembers(const FPlayerReservation& OtherRes);
 
 	/** Is this data well formed */
-	bool IsValid() const;
+	bool IsValid(bool bIsValidationStrRequired=true) const;
 
 	/** Dump this reservation to log */
 	void Dump() const;
diff --git a/Engine/Plugins/Online/OnlineSubsystemUtils/Source/OnlineSubsystemUtils/Public/SpectatorBeaconHost.h b/Engine/Plugins/Online/OnlineSubsystemUtils/Source/OnlineSubsystemUtils/Public/SpectatorBeaconHost.h
--- a/Engine/Plugins/Online/OnlineSubsystemUtils/Source/OnlineSubsystemUtils/Public/SpectatorBeaconHost.h
+++ b/Engine/Plugins/Online/OnlineSubsystemUtils/Source/OnlineSubsystemUtils/Public/SpectatorBeaconHost.h
@@ -240,6 +240,9 @@ protected:
 	/** Do the timeouts below cause a player to be removed from the reservation list */
 	UPROPERTY(Config)
 		bool bLogoutOnSessionTimeout;
+	/** Do spectators require validation strings. */
+	UPROPERTY(Config)
+		bool bIsValidationStrRequired;
 	/** Seconds that can elapse before a reservation is removed due to player not being registered with the session */
 	UPROPERTY(Transient, Config)
 		float SessionTimeoutSecs;
diff --git a/Engine/Plugins/Online/OnlineSubsystemUtils/Source/OnlineSubsystemUtils/Public/SpectatorBeaconState.h b/Engine/Plugins/Online/OnlineSubsystemUtils/Source/OnlineSubsystemUtils/Public/SpectatorBeaconState.h
--- a/Engine/Plugins/Online/OnlineSubsystemUtils/Source/OnlineSubsystemUtils/Public/SpectatorBeaconState.h
+++ b/Engine/Plugins/Online/OnlineSubsystemUtils/Source/OnlineSubsystemUtils/Public/SpectatorBeaconState.h
@@ -175,7 +175,7 @@ struct ONLINESUBSYSTEMUTILS_API FSpectatorReservation
 		FPlayerReservation Spectator;
 
 	/** Is this data well formed */
-	bool IsValid() const;
+	bool IsValid(bool bIsValidationStrRequired=true) const;
 
 	/** Dump this reservation to log */
 	void Dump() const;
