Prevent ensure to trigger if GetMicrophoneAmplitude is not implemented

diff --git a/Engine/Plugins/Online/OnlineSubsystemUtils/Source/OnlineSubsystemUtils/Private/VoiceInterfaceImpl.cpp b/Engine/Plugins/Online/OnlineSubsystemUtils/Source/OnlineSubsystemUtils/Private/VoiceInterfaceImpl.cpp
--- a/Engine/Plugins/Online/OnlineSubsystemUtils/Source/OnlineSubsystemUtils/Private/VoiceInterfaceImpl.cpp
+++ b/Engine/Plugins/Online/OnlineSubsystemUtils/Source/OnlineSubsystemUtils/Private/VoiceInterfaceImpl.cpp
@@ -729,13 +729,13 @@ void FOnlineVoiceImpl::ProcessLocalVoicePackets()
 						// Process this user
 						uint32 Result = VoiceEngine->ReadLocalVoiceData(Index, BufferStart, &SpaceAvail, &SampleCount);
 
-						// Convert to Q15:
 						float Amplitude = VoiceEngine->GetMicrophoneAmplitude(Index);
+						Amplitude = Amplitude < 0.0f ? 1.0f : Amplitude; // GetMicrophoneAmplitude returns -1 if not implemented: we set it to a sane value so we don't trigger the ensure
 						if (!ensure(Amplitude >= 0.0f && Amplitude <= 1.0f))
 						{
-							// GetMicrophoneAmplitude returns -1 if not implemented which would mess up the MicrophoneAmplitude value so we set it to a sane value
-							Amplitude = 1.0f;
+							Amplitude = 1.0f; // make sure Amplitude has a sane value in any case
 						}
+						// Convert to Q15:
 						VoiceData.LocalPackets[Index].MicrophoneAmplitude = (int16)(Amplitude * 32767.0f);
 
 						if (Result == ONLINE_SUCCESS)
