AutoRTFM some missing modules

diff --git a/Engine/Plugins/Experimental/ControlFlows/Source/ControlFlows/ControlFlows.Build.cs b/Engine/Plugins/Experimental/ControlFlows/Source/ControlFlows/ControlFlows.Build.cs
--- a/Engine/Plugins/Experimental/ControlFlows/Source/ControlFlows/ControlFlows.Build.cs
+++ b/Engine/Plugins/Experimental/ControlFlows/Source/ControlFlows/ControlFlows.Build.cs
@@ -14,5 +14,7 @@ public class ControlFlows : ModuleRules
 				// ... add private dependencies that you statically link with here ...	
 			}
 			);
+
+		bAllowAutoRTFMInstrumentation = true;
 	}
 }
diff --git a/Engine/Plugins/Online/OnlineSubsystemUtils/Source/OnlineSubsystemUtils/OnlineSubsystemUtils.Build.cs b/Engine/Plugins/Online/OnlineSubsystemUtils/Source/OnlineSubsystemUtils/OnlineSubsystemUtils.Build.cs
--- a/Engine/Plugins/Online/OnlineSubsystemUtils/Source/OnlineSubsystemUtils/OnlineSubsystemUtils.Build.cs
+++ b/Engine/Plugins/Online/OnlineSubsystemUtils/Source/OnlineSubsystemUtils/OnlineSubsystemUtils.Build.cs
@@ -56,5 +56,7 @@ public class OnlineSubsystemUtils : ModuleRules
 		);
 
 		PublicDependencyModuleNames.Add("OnlineSubsystem");
+
+		bAllowAutoRTFMInstrumentation = true;
 	}
 }
diff --git a/Engine/Source/Runtime/PakFile/PakFile.Build.cs b/Engine/Source/Runtime/PakFile/PakFile.Build.cs
--- a/Engine/Source/Runtime/PakFile/PakFile.Build.cs
+++ b/Engine/Source/Runtime/PakFile/PakFile.Build.cs
@@ -10,5 +10,7 @@ public class PakFile : ModuleRules
 		PrivateDependencyModuleNames.Add("Core");
 		PrivateDependencyModuleNames.Add("CoreUObject");
 		PublicDependencyModuleNames.Add("RSA");
+
+		bAllowAutoRTFMInstrumentation = true;
 	}
 }
