Add a virtual DoesModifiedReservationFit method which behaves similar to DoesReservationFit for reservation modifications.

diff --git a/Engine/Plugins/Online/OnlineSubsystemUtils/Source/OnlineSubsystemUtils/Private/PartyBeaconHost.cpp b/Engine/Plugins/Online/OnlineSubsystemUtils/Source/OnlineSubsystemUtils/Private/PartyBeaconHost.cpp
--- a/Engine/Plugins/Online/OnlineSubsystemUtils/Source/OnlineSubsystemUtils/Private/PartyBeaconHost.cpp
+++ b/Engine/Plugins/Online/OnlineSubsystemUtils/Source/OnlineSubsystemUtils/Private/PartyBeaconHost.cpp
@@ -894,7 +894,7 @@ EPartyReservationResult::Type APartyBeaconHost::UpdatePartyReservation(const FPa
 				}
 				// check to see if we have space to add new reservations for the new players 
 				// Not using IsBeaconFull as we may not be adding a new player in the situation where a party player who has a reservation joins the game in which case NewPlayers.Num == 0
-				if ((State->GetRemainingReservations() - NewPlayers.Num()) >= 0)
+				if (State->DoesModifiedReservationFit(ExistingReservation, NewPlayers))
 				{
 					// Validate that adding the new party members to this reservation entry still fits within the team size
 					if (!ShouldRespectCompetitiveIntegrity() || (NewPlayers.Num() - NumPlayersWithExistingReservation) <= NumAvailableSlotsOnTeam)
diff --git a/Engine/Plugins/Online/OnlineSubsystemUtils/Source/OnlineSubsystemUtils/Private/PartyBeaconState.cpp b/Engine/Plugins/Online/OnlineSubsystemUtils/Source/OnlineSubsystemUtils/Private/PartyBeaconState.cpp
--- a/Engine/Plugins/Online/OnlineSubsystemUtils/Source/OnlineSubsystemUtils/Private/PartyBeaconState.cpp
+++ b/Engine/Plugins/Online/OnlineSubsystemUtils/Source/OnlineSubsystemUtils/Private/PartyBeaconState.cpp
@@ -641,6 +641,11 @@ bool UPartyBeaconState::DoesReservationFit(const FPartyReservation& ReservationR
 	return bPartySizeOk && bRoomForReservation;
 }
 
+bool UPartyBeaconState::DoesModifiedReservationFit(const FPartyReservation& ExistingReservation, const TArray<FPlayerReservation>& NewPlayers) const
+{
+	return GetRemainingReservations() - NewPlayers.Num() >= 0;
+}
+
 bool UPartyBeaconState::AddReservation(const FPartyReservation& ReservationRequest)
 {
 	if (UE_LOG_ACTIVE(LogPartyBeacon, Verbose))
diff --git a/Engine/Plugins/Online/OnlineSubsystemUtils/Source/OnlineSubsystemUtils/Public/PartyBeaconState.h b/Engine/Plugins/Online/OnlineSubsystemUtils/Source/OnlineSubsystemUtils/Public/PartyBeaconState.h
--- a/Engine/Plugins/Online/OnlineSubsystemUtils/Source/OnlineSubsystemUtils/Public/PartyBeaconState.h
+++ b/Engine/Plugins/Online/OnlineSubsystemUtils/Source/OnlineSubsystemUtils/Public/PartyBeaconState.h
@@ -401,6 +401,16 @@ class ONLINESUBSYSTEMUTILS_API UPartyBeaconState : public UObject
 	 */
 	virtual bool DoesReservationFit(const FPartyReservation& ReservationRequest) const;
 
+	/**
+	 * Determine if this reservation fits all rules for fitting in the game
+	 * 
+	 * @param ExistingReservation reservation to test for
+	 * @param NewPlayers players requested to be added
+	 *
+	 * @return true if reservation can be modified to fit this state, false otherwise
+	 */
+	virtual bool DoesModifiedReservationFit(const FPartyReservation& ExistingReservation, const TArray<FPlayerReservation>& NewPlayers) const;
+
 	/**
 	 * @return true if the beacon is currently at max capacity
 	 */
