Make party session reservation timeout be overrideable

diff --git a/Engine/Plugins/Online/OnlineSubsystemUtils/Source/OnlineSubsystemUtils/Private/PartyBeaconHost.cpp b/Engine/Plugins/Online/OnlineSubsystemUtils/Source/OnlineSubsystemUtils/Private/PartyBeaconHost.cpp
--- a/Engine/Plugins/Online/OnlineSubsystemUtils/Source/OnlineSubsystemUtils/Private/PartyBeaconHost.cpp
+++ b/Engine/Plugins/Online/OnlineSubsystemUtils/Source/OnlineSubsystemUtils/Private/PartyBeaconHost.cpp
@@ -193,7 +193,7 @@ void APartyBeaconHost::Tick(float DeltaTime)
 									// if the timeout has been exceeded then add to list of players 
 									// that need to be logged out from the beacon
 									if ((bIsPlayerPendingJoin && PlayerEntry.ElapsedTime > TravelSessionTimeoutSecs) ||
-										(!bIsPlayerPendingJoin && PlayerEntry.ElapsedTime > SessionTimeoutSecs))
+										(!bIsPlayerPendingJoin && PlayerEntry.ElapsedTime > GetSessionTimeoutSecs(PlayerEntry.UniqueId)))
 									{
 										PlayersToLogout.AddUnique(PlayerEntry.UniqueId.GetUniqueNetId());
 									}
@@ -483,6 +483,11 @@ bool APartyBeaconHost::GetPlayerValidation(const FUniqueNetId& PlayerId, FString
 	return bHasValidation;
 }
 
+float APartyBeaconHost::GetSessionTimeoutSecs(const FUniqueNetIdRepl& PlayerId) const
+{
+	return SessionTimeoutSecs;
+}
+
 bool APartyBeaconHost::GetPartyLeader(const FUniqueNetIdRepl& InPartyMemberId, FUniqueNetIdRepl& OutPartyLeaderId) const
 {
 	bool bHasLeader = false;
diff --git a/Engine/Plugins/Online/OnlineSubsystemUtils/Source/OnlineSubsystemUtils/Public/PartyBeaconHost.h b/Engine/Plugins/Online/OnlineSubsystemUtils/Source/OnlineSubsystemUtils/Public/PartyBeaconHost.h
--- a/Engine/Plugins/Online/OnlineSubsystemUtils/Source/OnlineSubsystemUtils/Public/PartyBeaconHost.h
+++ b/Engine/Plugins/Online/OnlineSubsystemUtils/Source/OnlineSubsystemUtils/Public/PartyBeaconHost.h
@@ -242,6 +242,15 @@ class ONLINESUBSYSTEMUTILS_API APartyBeaconHost : public AOnlineBeaconHostObject
 	 */
 	virtual bool GetPlayerValidation(const FUniqueNetId& PlayerId, FString& OutValidation) const;
 
+	/**
+	 * Get the session timeout value to use for the target player
+	 *
+	 * @param PlayerId unique id of player to find session timeout for
+	 *
+	 * @return The value to sue as the session timeout
+	 */
+	virtual float GetSessionTimeoutSecs(const FUniqueNetIdRepl& PlayerId) const;
+
 	/**
 	 * Get the party leader for a given unique id
 	 *
