OnlineBeaconClient now cleans up properly if the beacon's NetDriver is pending destruction.

diff --git a/Engine/Plugins/Online/OnlineSubsystemUtils/Source/OnlineSubsystemUtils/Private/OnlineBeaconClient.cpp b/Engine/Plugins/Online/OnlineSubsystemUtils/Source/OnlineSubsystemUtils/Private/OnlineBeaconClient.cpp
--- a/Engine/Plugins/Online/OnlineSubsystemUtils/Source/OnlineSubsystemUtils/Private/OnlineBeaconClient.cpp
+++ b/Engine/Plugins/Online/OnlineSubsystemUtils/Source/OnlineSubsystemUtils/Private/OnlineBeaconClient.cpp
@@ -58,7 +58,11 @@ class UNetConnection* AOnlineBeaconClient::GetNetConnection() const
 
 bool AOnlineBeaconClient::DestroyNetworkActorHandled()
 {
-	if (BeaconConnection && BeaconConnection->GetConnectionState() != USOCK_Closed)
+	// AOnlineBeacon::NetDriver may have already been nulled out in AOnlineBeacon::CleanupNetDriver
+	const UNetDriver* LocalNetDriver = BeaconConnection ? BeaconConnection->Driver : nullptr;
+
+	if (BeaconConnection && BeaconConnection->GetConnectionState() != USOCK_Closed &&
+		LocalNetDriver && !LocalNetDriver->GetPendingDestruction())
 	{
 		// This will be cleaned up in ~2 sec by UNetConnection Tick
 		BeaconConnection->bPendingDestroy = true;
