Fix wrong check of received NetId in Show External Login UI blueprint

diff --git a/Engine/Plugins/Online/OnlineSubsystemUtils/Source/OnlineSubsystemUtils/Private/ShowLoginUICallbackProxy.cpp b/Engine/Plugins/Online/OnlineSubsystemUtils/Source/OnlineSubsystemUtils/Private/ShowLoginUICallbackProxy.cpp
--- a/Engine/Plugins/Online/OnlineSubsystemUtils/Source/OnlineSubsystemUtils/Private/ShowLoginUICallbackProxy.cpp
+++ b/Engine/Plugins/Online/OnlineSubsystemUtils/Source/OnlineSubsystemUtils/Private/ShowLoginUICallbackProxy.cpp
@@ -87,7 +87,7 @@ void UShowLoginUICallbackProxy::OnShowLoginUICompleted(FUniqueNetIdPtr UniqueId,
 		}
 	}
 
-	if (UniqueId.IsValid())
+	if (UniqueId.IsValid() && UniqueId->IsValid())
 	{
 		OnSuccess.Broadcast(MyPlayerController);
 	}
