Validate party beacon connection owner for reservations

diff --git a/Engine/Plugins/Online/OnlineSubsystemUtils/Source/OnlineSubsystemUtils/Private/PartyBeaconClient.cpp b/Engine/Plugins/Online/OnlineSubsystemUtils/Source/OnlineSubsystemUtils/Private/PartyBeaconClient.cpp
--- a/Engine/Plugins/Online/OnlineSubsystemUtils/Source/OnlineSubsystemUtils/Private/PartyBeaconClient.cpp
+++ b/Engine/Plugins/Online/OnlineSubsystemUtils/Source/OnlineSubsystemUtils/Private/PartyBeaconClient.cpp
@@ -8,9 +8,9 @@
 
 #include UE_INLINE_GENERATED_CPP_BY_NAME(PartyBeaconClient)
 
-#if !UE_BUILD_SHIPPING
 namespace BeaconConsoleVariables
 {
+#if !UE_BUILD_SHIPPING
 	/** Time to delay delegates firing a reservation request response */
 	TAutoConsoleVariable<float> CVarDelayReservationResponse(
 		TEXT("beacon.DelayReservationResponse"),
@@ -42,9 +42,27 @@ namespace BeaconConsoleVariables
 		TEXT("Delay time between received full response and notification\n")
 		TEXT("Time in secs"),
 		ECVF_Default);
-}
 #endif
 
+	/** Whether to restrict party reservation RPCs to require that the initiator is the party leader. */
+	TAutoConsoleVariable<bool> CVarRequireInitiatorIsPartyLeader(
+		TEXT("beacon.RequireInitiatorIsPartyLeader"),
+		true,
+		TEXT("Enforce RPC validation which checks whether the initiator of a reservation RPC is the party leader\n")
+		TEXT("Enabled"),
+		ECVF_Default);
+}
+
+namespace
+{
+	bool IsRPCInitiatorValid(const FUniqueNetIdRepl& BeaconOwner, const FUniqueNetIdRepl& RequestPartyLeader)
+	{
+		const bool bInitiatorValidationEnabled = BeaconConsoleVariables::CVarRequireInitiatorIsPartyLeader.GetValueOnGameThread();
+		const bool bInitiatorIsValid = bInitiatorValidationEnabled ? (RequestPartyLeader.IsValid() && RequestPartyLeader == BeaconOwner) : true;
+		return bInitiatorIsValid;
+	}
+}
+
 /** Max time to wait for a response from the server for CancelReservation */
 #define CANCEL_FAILSAFE 5.0f
 
@@ -379,7 +397,9 @@ void APartyBeaconClient::OnFailure()
 
 bool APartyBeaconClient::ServerReservationRequest_Validate(const FString& SessionId, const FPartyReservation& Reservation)
 {
-	return !SessionId.IsEmpty() && Reservation.PartyLeader.IsValid() && Reservation.PartyMembers.Num() > 0;
+	const bool bInitiatorIsValid = IsRPCInitiatorValid(GetUniqueId(), Reservation.PartyLeader);
+	const bool bRequestValid = bInitiatorIsValid && !SessionId.IsEmpty() && Reservation.PartyLeader.IsValid() && Reservation.PartyMembers.Num() > 0;
+	return bRequestValid;
 }
 
 void APartyBeaconClient::ServerReservationRequest_Implementation(const FString& SessionId, const FPartyReservation& Reservation)
@@ -395,7 +415,9 @@ void APartyBeaconClient::ServerReservationRequest_Implementation(const FString&
 
 bool APartyBeaconClient::ServerUpdateReservationRequest_Validate(const FString& SessionId, const FPartyReservation& ReservationUpdate)
 {
-	return !SessionId.IsEmpty() && ReservationUpdate.PartyLeader.IsValid() && ReservationUpdate.PartyMembers.Num() > 0;
+	const bool bInitiatorIsValid = IsRPCInitiatorValid(GetUniqueId(), ReservationUpdate.PartyLeader);
+	const bool bRequestValid = bInitiatorIsValid && !SessionId.IsEmpty() && ReservationUpdate.PartyLeader.IsValid() && ReservationUpdate.PartyMembers.Num() > 0;
+	return bRequestValid;
 }
 
 void APartyBeaconClient::ServerUpdateReservationRequest_Implementation(const FString& SessionId, const FPartyReservation& ReservationUpdate)
@@ -411,7 +433,9 @@ void APartyBeaconClient::ServerUpdateReservationRequest_Implementation(const FSt
 
 bool APartyBeaconClient::ServerAddOrUpdateReservationRequest_Validate(const FString& SessionId, const FPartyReservation& Reservation)
 {
-	return !SessionId.IsEmpty() && Reservation.PartyLeader.IsValid() && Reservation.PartyMembers.Num() > 0;
+	const bool bInitiatorIsValid = IsRPCInitiatorValid(GetUniqueId(), Reservation.PartyLeader);
+	const bool bRequestValid = bInitiatorIsValid && !SessionId.IsEmpty() && Reservation.PartyLeader.IsValid() && Reservation.PartyMembers.Num() > 0;
+	return bRequestValid;
 }
 
 void APartyBeaconClient::ServerAddOrUpdateReservationRequest_Implementation(const FString& SessionId, const FPartyReservation& Reservation)
@@ -426,7 +450,9 @@ void APartyBeaconClient::ServerAddOrUpdateReservationRequest_Implementation(cons
 
 bool APartyBeaconClient::ServerRemoveMemberFromReservationRequest_Validate(const FString& SessionId, const FPartyReservation& ReservationUpdate)
 {
-	return !SessionId.IsEmpty() && ReservationUpdate.PartyLeader.IsValid() && ReservationUpdate.PartyMembers.Num() > 0;
+	const bool bInitiatorIsValid = IsRPCInitiatorValid(GetUniqueId(), ReservationUpdate.PartyLeader);
+	const bool bRequestValid = bInitiatorIsValid && !SessionId.IsEmpty() && ReservationUpdate.PartyLeader.IsValid() && ReservationUpdate.PartyMembers.Num() > 0;
+	return bRequestValid;
 }
 
 void APartyBeaconClient::ServerRemoveMemberFromReservationRequest_Implementation(const FString& SessionId, const FPartyReservation& ReservationUpdate)
@@ -442,7 +468,9 @@ void APartyBeaconClient::ServerRemoveMemberFromReservationRequest_Implementation
 
 bool APartyBeaconClient::ServerCancelReservationRequest_Validate(const FUniqueNetIdRepl& PartyLeader)
 {
-	return true;
+	const bool bInitiatorIsValid = IsRPCInitiatorValid(GetUniqueId(), PartyLeader);
+	const bool bRequestValid = bInitiatorIsValid && PartyLeader.IsValid();
+	return bRequestValid;
 }
 
 void APartyBeaconClient::ServerCancelReservationRequest_Implementation(const FUniqueNetIdRepl& PartyLeader)
@@ -634,4 +662,3 @@ FTimerHandle APartyBeaconClient::DelayResponse(FTimerDelegate& Delegate, float D
 
 	return TimerHandle;
 }
-
