Add missing network flush after client sends NMT_BeaconNetGUIDAck.
Add missing network flush after host sends NMT_BeaconAssignGUID and ClientOnConnected RPC.
Fix FNetControlMessageInfo to handle intializing message type names in editor builds.
Add unit tests for OnlineBeacon.

diff --git a/Engine/Plugins/Online/OnlineSubsystemUtils/OnlineSubsystemUtils.uplugin b/Engine/Plugins/Online/OnlineSubsystemUtils/OnlineSubsystemUtils.uplugin
--- a/Engine/Plugins/Online/OnlineSubsystemUtils/OnlineSubsystemUtils.uplugin
+++ b/Engine/Plugins/Online/OnlineSubsystemUtils/OnlineSubsystemUtils.uplugin
@@ -1,27 +1,33 @@
 {
-	"FileVersion" : 3,
-	
-	"FriendlyName" : "Online Subsystem Utils",
-	"Version" : 1,
-	"VersionName" : "1.0",
-	"Description" : "Shared code for interacting online service and online subsystem implementations.",
-	"Category" : "Online Platform",
-	"CreatedBy" : "Epic Games, Inc.",
-	"CreatedByURL" : "https://epicgames.com",
-	"EnabledByDefault" : true,
+	"FileVersion": 3,
 
-    "Modules": [
-        {
-            "Name": "OnlineSubsystemUtils",
-            "Type": "Runtime",
-            "LoadingPhase": "PostConfigInit"
-        },
-        {
-            "Name": "OnlineBlueprintSupport",
-            "Type": "UncookedOnly",
-            "LoadingPhase": "PostConfigInit"
-        }
-    ],
+	"FriendlyName": "Online Subsystem Utils",
+	"Version": 1,
+	"VersionName": "1.0",
+	"Description": "Shared code for interacting online service and online subsystem implementations.",
+	"Category": "Online Platform",
+	"CreatedBy": "Epic Games, Inc.",
+	"CreatedByURL": "https://epicgames.com",
+	"EnabledByDefault": true,
+
+	"Modules":
+	[
+		{
+			"Name": "OnlineSubsystemUtils",
+			"Type": "Runtime",
+			"LoadingPhase": "PostConfigInit"
+		},
+		{
+			"Name": "OnlineSubsystemUtilsTest",
+			"Type": "Editor",
+			"LoadingPhase": "PostEngineInit"
+		},
+		{
+			"Name": "OnlineBlueprintSupport",
+			"Type": "UncookedOnly",
+			"LoadingPhase": "PostConfigInit"
+		}
+	],
 	"LocalizationTargets":
 	[
 		{
diff --git a/Engine/Plugins/Online/OnlineSubsystemUtils/Source/OnlineSubsystemUtils/Private/OnlineBeaconClient.cpp b/Engine/Plugins/Online/OnlineSubsystemUtils/Source/OnlineSubsystemUtils/Private/OnlineBeaconClient.cpp
--- a/Engine/Plugins/Online/OnlineSubsystemUtils/Source/OnlineSubsystemUtils/Private/OnlineBeaconClient.cpp
+++ b/Engine/Plugins/Online/OnlineSubsystemUtils/Source/OnlineSubsystemUtils/Private/OnlineBeaconClient.cpp
@@ -445,7 +445,7 @@ void AOnlineBeaconClient::NotifyControlMessage(UNetConnection* Connection, uint8
 				// send NMT_Login
 				Connection->ClientResponse = TEXT("0");
 				FNetControlMessage<NMT_Login>::Send(Connection, Connection->ClientResponse, URLString, Connection->PlayerId, OnlinePlatformNameString);
-				NetDriver->ServerConnection->FlushNet();
+				Connection->FlushNet();
 			}
 			else
 			{
@@ -464,7 +464,7 @@ void AOnlineBeaconClient::NotifyControlMessage(UNetConnection* Connection, uint8
 				if (!BeaconType.IsEmpty())
 				{
 					FNetControlMessage<NMT_BeaconJoin>::Send(Connection, BeaconType, Connection->PlayerId);
-					NetDriver->ServerConnection->FlushNet();
+					Connection->FlushNet();
 				}
 				else
 				{
@@ -486,6 +486,7 @@ void AOnlineBeaconClient::NotifyControlMessage(UNetConnection* Connection, uint8
 
 						FString BeaconType = GetBeaconType();
 						FNetControlMessage<NMT_BeaconNetGUIDAck>::Send(Connection, BeaconType);
+						Connection->FlushNet();
 						// Server will send ClientOnConnected() when it gets this control message
 
 						// Fail safe for connection to server but no client connection RPC
diff --git a/Engine/Plugins/Online/OnlineSubsystemUtils/Source/OnlineSubsystemUtils/Private/OnlineBeaconHost.cpp b/Engine/Plugins/Online/OnlineSubsystemUtils/Source/OnlineSubsystemUtils/Private/OnlineBeaconHost.cpp
--- a/Engine/Plugins/Online/OnlineSubsystemUtils/Source/OnlineSubsystemUtils/Private/OnlineBeaconHost.cpp
+++ b/Engine/Plugins/Online/OnlineSubsystemUtils/Source/OnlineSubsystemUtils/Private/OnlineBeaconHost.cpp
@@ -472,6 +472,7 @@ bool AOnlineBeaconHost::HandleControlMessage(UNetConnection* Connection, uint8 M
 			ClientActors.Add(NewClientActor);
 
 			FNetControlMessage<NMT_BeaconAssignGUID>::Send(Connection, NetGUID);
+			Connection->FlushNet(true);
 		}
 		break;
 	case NMT_BeaconNetGUIDAck:
@@ -563,6 +564,8 @@ void AOnlineBeaconHost::FinishHandshake(UNetConnection* Connection, FString Beac
 
 	// Send an RPC to the client to open the actor channel and guarantee RPCs will work
 	ClientActor->ClientOnConnected();
+	Connection->FlushNet(true);
+
 	UE_LOG(LogBeacon, Log, TEXT("%s: Handshake complete."), *GetDebugName(Connection));
 	OnBeaconConnectedDelegate->ExecuteIfBound(ClientActor, Connection);
 }
diff --git a/Engine/Plugins/Online/OnlineSubsystemUtils/Source/OnlineSubsystemUtilsTest/OnlineSubsystemUtilsTest.Build.cs b/Engine/Plugins/Online/OnlineSubsystemUtils/Source/OnlineSubsystemUtilsTest/OnlineSubsystemUtilsTest.Build.cs
new file mode 100644
--- /dev/null
+++ b/Engine/Plugins/Online/OnlineSubsystemUtils/Source/OnlineSubsystemUtilsTest/OnlineSubsystemUtilsTest.Build.cs
@@ -0,0 +1,31 @@
+// Copyright Epic Games, Inc. All Rights Reserved.
+
+using UnrealBuildTool;
+using System.IO;
+
+public class OnlineSubsystemUtilsTest : ModuleRules
+{
+	public OnlineSubsystemUtilsTest(ReadOnlyTargetRules Target) : base(Target)
+	{
+		PublicDependencyModuleNames.AddRange(
+			new string[]
+			{
+				"Core",
+				"CoreOnline",
+				"CoreUObject",
+				"Engine",
+				"NetCore",
+				"UnrealEd",
+				"OnlineSubsystem",
+				"OnlineSubsystemUtils"
+			}
+		);
+
+		PrivateDependencyModuleNames.AddRange(
+			new string[]
+			{
+				"Sockets"
+			}
+		);
+	}
+}
diff --git a/Engine/Plugins/Online/OnlineSubsystemUtils/Source/OnlineSubsystemUtilsTest/Private/OnlineBeaconUnitTest.cpp b/Engine/Plugins/Online/OnlineSubsystemUtils/Source/OnlineSubsystemUtilsTest/Private/OnlineBeaconUnitTest.cpp
new file mode 100644
--- /dev/null
+++ b/Engine/Plugins/Online/OnlineSubsystemUtils/Source/OnlineSubsystemUtilsTest/Private/OnlineBeaconUnitTest.cpp
@@ -0,0 +1,3262 @@
+// Copyright Epic Games, Inc. All Rights Reserved.
+
+#include "Misc/AutomationTest.h"
+
+#include "Algo/RemoveIf.h"
+#include "Engine/Engine.h"
+#include "Engine/World.h"
+#include "Misc/Base64.h"
+#include "Net/DataChannel.h"
+#include "OnlineIdentityErrors.h"
+#include "OnlineBeaconUnitTestClient.h"
+#include "OnlineBeaconUnitTestHost.h"
+#include "OnlineBeaconUnitTestHostObject.h"
+#include "OnlineBeaconUnitTestUtils.h"
+#include "OnlineSubsystemTypes.h"
+
+IMPLEMENT_SIMPLE_AUTOMATION_TEST(FOnlineBeaconTestBasicHandshakeSuccess,
+	"System.Engine.Online.OnlineSubsystemUtils.OnlineBeacon.BasicHandshakeSuccess",
+	EAutomationTestFlags::EditorContext | EAutomationTestFlags::EngineFilter)
+bool FOnlineBeaconTestBasicHandshakeSuccess::RunTest(const FString& Parameters)
+{
+	TSharedPtr<BeaconUnitTest::FTestPrerequisites> Prerequisites = BeaconUnitTest::FTestPrerequisites::TryCreate();
+	UTEST_TRUE_EXPR(Prerequisites.IsValid());
+
+	Prerequisites->GetConfig().NetDriver.bNoTimeouts = true;
+	Prerequisites->GetConfig().NetDriver.ServerListenPort = 9999;
+
+	// Snapshot config to restore before running a test section.
+	const BeaconUnitTest::FTestConfig BaseConfig = Prerequisites->GetConfig();
+
+	FUniqueNetIdStringRef UserId = FUniqueNetIdString::Create(TEXT("User"), TEXT("UnitTest"));
+
+	// Host setup.
+	TObjectPtr<AOnlineBeaconUnitTestHostObject> BeaconHostObject = Prerequisites->GetWorld()->SpawnActor<AOnlineBeaconUnitTestHostObject>(AOnlineBeaconUnitTestHostObject::StaticClass());
+	TObjectPtr<AOnlineBeaconUnitTestHost> BeaconHost = Prerequisites->GetWorld()->SpawnActor<AOnlineBeaconUnitTestHost>(AOnlineBeaconUnitTestHost::StaticClass());
+	ON_SCOPE_EXIT{ if (BeaconHostObject) { BeaconHostObject->Destroy(); } };
+	ON_SCOPE_EXIT{ if (BeaconHost) { BeaconHost->DestroyBeacon(); } };
+	UTEST_TRUE_EXPR(BeaconUnitTest::ConfigureBeacon(*Prerequisites, BeaconHost));
+	BeaconHost->ListenPort = Prerequisites->GetConfig().NetDriver.ServerListenPort;
+	BeaconHost->RegisterHost(BeaconHostObject);
+	UTEST_TRUE_EXPR(BeaconHost->InitHost());
+	BeaconHost->PauseBeaconRequests(false);
+
+	TSharedPtr<BeaconUnitTest::FNetworkStats> BeaconHostNetStats;
+	UTEST_TRUE_EXPR(BeaconUnitTest::ConfigureBeaconNetDriver(*Prerequisites, BeaconHost, &BeaconHostNetStats));
+	UTEST_TRUE_EXPR(BeaconHostNetStats.IsValid());
+
+	// Successful handshake.
+	// Disconnect initiated by client.
+	{
+		BeaconHostNetStats->ReceivedControlMessages.Empty();
+		Prerequisites->GetConfig() = BaseConfig;
+		Prerequisites->GetStats() = BeaconUnitTest::FTestStats();
+
+		TObjectPtr<AOnlineBeaconClient> HostUserBeaconClient;
+
+		// Client setup.
+
+		TObjectPtr<AOnlineBeaconUnitTestClient> BeaconClient = Prerequisites->GetWorld()->SpawnActor<AOnlineBeaconUnitTestClient>(AOnlineBeaconUnitTestClient::StaticClass());
+		ON_SCOPE_EXIT{ if (BeaconClient) { BeaconClient->DestroyBeacon(); } };
+		UTEST_TRUE_EXPR(BeaconUnitTest::ConfigureBeacon(*Prerequisites, BeaconClient));
+
+		UTEST_TRUE_EXPR(BeaconClient->GetConnectionState() == EBeaconConnectionState::Invalid);
+		UTEST_TRUE_EXPR(BeaconUnitTest::InitClientForUser(*Prerequisites, BeaconClient, UserId));
+		UTEST_TRUE_EXPR(BeaconClient->GetConnectionState() == EBeaconConnectionState::Pending);
+		UTEST_TRUE_EXPR(BeaconClient->GetNetConnection() != nullptr);
+
+		TSharedPtr<BeaconUnitTest::FNetworkStats> BeaconClientNetStats;
+		UTEST_TRUE_EXPR(BeaconUnitTest::ConfigureBeaconNetDriver(*Prerequisites, BeaconClient, &BeaconClientNetStats));
+		UTEST_TRUE_EXPR(BeaconClientNetStats.IsValid());
+
+		// Handshake testing.
+
+		UTEST_TRUE_EXPR(BeaconUnitTest::TickUntilConnectionInitialized(*Prerequisites, BeaconClient));
+		UTEST_TRUE_EXPR(BeaconClient->GetConnectionState() == EBeaconConnectionState::Pending);
+		UTEST_TRUE_EXPR(BeaconClientNetStats->ReceivedControlMessages.IsEmpty());
+		UTEST_TRUE_EXPR(BeaconHostNetStats->ReceivedControlMessages.IsEmpty());
+		UTEST_TRUE_EXPR(Prerequisites->GetStats().Client.OnConnected.InvokeCount == 0);
+		UTEST_TRUE_EXPR(Prerequisites->GetStats().Client.OnFailure.InvokeCount == 0);
+		UTEST_TRUE_EXPR(Prerequisites->GetStats().Host.OnFailure.InvokeCount == 0);
+		UTEST_TRUE_EXPR(Prerequisites->GetStats().HostObject.OnClientConnected.InvokeCount == 0);
+		UTEST_TRUE_EXPR(Prerequisites->GetStats().Encryption.NetworkEncryptionFailure.InvokeCount == 1);
+
+		UTEST_TRUE_EXPR(BeaconUnitTest::TickUntilControlMessageReceived(*Prerequisites, *BeaconHostNetStats));
+		UTEST_TRUE_EXPR(BeaconClient->GetConnectionState() == EBeaconConnectionState::Pending);
+		UTEST_TRUE_EXPR(BeaconClientNetStats->ReceivedControlMessages.IsEmpty());
+		UTEST_TRUE_EXPR(BeaconHostNetStats->ReceivedControlMessages == TArray<uint8>({ NMT_Hello }));
+		UTEST_TRUE_EXPR(Prerequisites->GetStats().Client.OnConnected.InvokeCount == 0);
+		UTEST_TRUE_EXPR(Prerequisites->GetStats().Client.OnFailure.InvokeCount == 0);
+		UTEST_TRUE_EXPR(Prerequisites->GetStats().Host.OnFailure.InvokeCount == 0);
+		UTEST_TRUE_EXPR(Prerequisites->GetStats().HostObject.OnClientConnected.InvokeCount == 0);
+		UTEST_TRUE_EXPR(Prerequisites->GetStats().Encryption.NetworkEncryptionFailure.InvokeCount == 2);
+		HostUserBeaconClient = BeaconUnitTest::GetBeaconClientForUser(BeaconHostObject, UserId);
+		UTEST_TRUE_EXPR(HostUserBeaconClient == nullptr);
+		BeaconHostNetStats->ReceivedControlMessages.Empty();
+
+		UTEST_TRUE_EXPR(BeaconUnitTest::TickOnce(*Prerequisites));
+		UTEST_TRUE_EXPR(BeaconClient->GetConnectionState() == EBeaconConnectionState::Pending);
+		UTEST_TRUE_EXPR(BeaconClientNetStats->ReceivedControlMessages == TArray<uint8>({ NMT_BeaconWelcome }));
+		UTEST_TRUE_EXPR(BeaconHostNetStats->ReceivedControlMessages.IsEmpty());
+		UTEST_TRUE_EXPR(Prerequisites->GetStats().Client.OnConnected.InvokeCount == 0);
+		UTEST_TRUE_EXPR(Prerequisites->GetStats().Client.OnFailure.InvokeCount == 0);
+		UTEST_TRUE_EXPR(Prerequisites->GetStats().Host.OnFailure.InvokeCount == 0);
+		UTEST_TRUE_EXPR(Prerequisites->GetStats().HostObject.OnClientConnected.InvokeCount == 0);
+		BeaconClientNetStats->ReceivedControlMessages.Empty();
+
+		UTEST_TRUE_EXPR(BeaconUnitTest::TickOnce(*Prerequisites));
+		UTEST_TRUE_EXPR(BeaconClient->GetConnectionState() == EBeaconConnectionState::Pending);
+		UTEST_TRUE_EXPR(BeaconClientNetStats->ReceivedControlMessages.IsEmpty());
+		UTEST_TRUE_EXPR(BeaconHostNetStats->ReceivedControlMessages == TArray<uint8>({ NMT_Netspeed, NMT_BeaconJoin }));
+		UTEST_TRUE_EXPR(Prerequisites->GetStats().Client.OnConnected.InvokeCount == 0);
+		UTEST_TRUE_EXPR(Prerequisites->GetStats().Client.OnFailure.InvokeCount == 0);
+		UTEST_TRUE_EXPR(Prerequisites->GetStats().Host.OnFailure.InvokeCount == 0);
+		UTEST_TRUE_EXPR(Prerequisites->GetStats().HostObject.OnClientConnected.InvokeCount == 0);
+		HostUserBeaconClient = BeaconUnitTest::GetBeaconClientForUser(BeaconHostObject, UserId);
+		UTEST_TRUE_EXPR(HostUserBeaconClient == nullptr);
+		BeaconHostNetStats->ReceivedControlMessages.Empty();
+
+		UTEST_TRUE_EXPR(BeaconUnitTest::TickOnce(*Prerequisites));
+		UTEST_TRUE_EXPR(BeaconClient->GetConnectionState() == EBeaconConnectionState::Pending);
+		UTEST_TRUE_EXPR(BeaconClientNetStats->ReceivedControlMessages == TArray<uint8>({ NMT_BeaconAssignGUID }));
+		UTEST_TRUE_EXPR(Prerequisites->GetStats().Client.OnConnected.InvokeCount == 0);
+		UTEST_TRUE_EXPR(Prerequisites->GetStats().Client.OnFailure.InvokeCount == 0);
+		UTEST_TRUE_EXPR(Prerequisites->GetStats().Host.OnFailure.InvokeCount == 0);
+		UTEST_TRUE_EXPR(Prerequisites->GetStats().HostObject.OnClientConnected.InvokeCount == 0);
+		UTEST_TRUE_EXPR(BeaconHostNetStats->ReceivedControlMessages.IsEmpty());
+		BeaconClientNetStats->ReceivedControlMessages.Empty();
+
+		UTEST_TRUE_EXPR(BeaconUnitTest::TickOnce(*Prerequisites));
+		UTEST_TRUE_EXPR(BeaconClient->GetConnectionState() == EBeaconConnectionState::Pending);
+		UTEST_TRUE_EXPR(BeaconClientNetStats->ReceivedControlMessages.IsEmpty());
+		UTEST_TRUE_EXPR(BeaconHostNetStats->ReceivedControlMessages == TArray<uint8>({ NMT_BeaconNetGUIDAck }));
+		UTEST_TRUE_EXPR(Prerequisites->GetStats().Client.OnConnected.InvokeCount == 0);
+		UTEST_TRUE_EXPR(Prerequisites->GetStats().Client.OnFailure.InvokeCount == 0);
+		UTEST_TRUE_EXPR(Prerequisites->GetStats().Host.OnFailure.InvokeCount == 0);
+		UTEST_TRUE_EXPR(Prerequisites->GetStats().HostObject.OnClientConnected.InvokeCount == 1);
+		HostUserBeaconClient = BeaconUnitTest::GetBeaconClientForUser(BeaconHostObject, UserId);
+		UTEST_TRUE_EXPR(HostUserBeaconClient != nullptr);
+		UTEST_TRUE_EXPR(HostUserBeaconClient->GetUniqueId() == *UserId);
+		BeaconHostNetStats->ReceivedControlMessages.Empty();
+
+		// Receive OnConnected RPC on the client (not a control message).
+		UTEST_TRUE_EXPR(BeaconUnitTest::TickOnce(*Prerequisites));
+		UTEST_TRUE_EXPR(BeaconClient->GetConnectionState() == EBeaconConnectionState::Open);
+		UTEST_TRUE_EXPR(Prerequisites->GetStats().Client.OnConnected.InvokeCount == 1);
+		UTEST_TRUE_EXPR(Prerequisites->GetStats().Client.OnFailure.InvokeCount == 0);
+		UTEST_TRUE_EXPR(Prerequisites->GetStats().Host.OnFailure.InvokeCount == 0);
+		UTEST_TRUE_EXPR(Prerequisites->GetStats().HostObject.OnClientConnected.InvokeCount == 1);
+		UTEST_TRUE_EXPR(Prerequisites->GetStats().HostObject.NotifyClientDisconnected.InvokeCount == 0);
+		UTEST_TRUE_EXPR(BeaconClientNetStats->ReceivedControlMessages.IsEmpty());
+		UTEST_TRUE_EXPR(BeaconHostNetStats->ReceivedControlMessages.IsEmpty());
+
+		// Handshake complete, channel is open.
+
+		// Disconnect the client.
+		BeaconClient->DestroyBeacon();
+		BeaconClient = nullptr;
+
+		// Check that client actor on the host cleaned up.
+		UTEST_TRUE_EXPR(BeaconUnitTest::TickUntil(*Prerequisites, [&Prerequisites](){ return Prerequisites->GetStats().HostObject.NotifyClientDisconnected.InvokeCount == 1; }));
+		UTEST_TRUE_EXPR(Prerequisites->GetStats().Client.OnConnected.InvokeCount == 1);
+		UTEST_TRUE_EXPR(Prerequisites->GetStats().Client.OnFailure.InvokeCount == 0);
+		UTEST_TRUE_EXPR(Prerequisites->GetStats().Host.OnFailure.InvokeCount == 0);
+		UTEST_TRUE_EXPR(Prerequisites->GetStats().HostObject.OnClientConnected.InvokeCount == 1);
+		UTEST_TRUE_EXPR(Prerequisites->GetStats().HostObject.NotifyClientDisconnected.InvokeCount == 1);
+		HostUserBeaconClient = BeaconUnitTest::GetBeaconClientForUser(BeaconHostObject, UserId);
+		UTEST_TRUE_EXPR(HostUserBeaconClient == nullptr);
+
+		// Check that encryption delegates were not fired.
+		UTEST_TRUE_EXPR(Prerequisites->GetStats().Encryption.NetworkEncryptionToken.InvokeCount == 0);
+		UTEST_TRUE_EXPR(Prerequisites->GetStats().Encryption.NetworkEncryptionToken.CallbackCount == 0);
+		UTEST_TRUE_EXPR(Prerequisites->GetStats().Encryption.NetworkEncryptionAck.InvokeCount == 0);
+		UTEST_TRUE_EXPR(Prerequisites->GetStats().Encryption.NetworkEncryptionAck.CallbackCount == 0);
+		// NetworkEncryptionFailure is always fired on client / host when no encryption data is provided.
+		UTEST_TRUE_EXPR(Prerequisites->GetStats().Encryption.NetworkEncryptionFailure.InvokeCount == 2);
+	}
+
+	// Successful handshake.
+	// Disconnect initiated by server.
+	{
+		BeaconHostNetStats->ReceivedControlMessages.Empty();
+		Prerequisites->GetConfig() = BaseConfig;
+		Prerequisites->GetStats() = BeaconUnitTest::FTestStats();
+
+		TObjectPtr<AOnlineBeaconClient> HostUserBeaconClient;
+
+		// Client setup.
+
+		TObjectPtr<AOnlineBeaconUnitTestClient> BeaconClient = Prerequisites->GetWorld()->SpawnActor<AOnlineBeaconUnitTestClient>(AOnlineBeaconUnitTestClient::StaticClass());
+		ON_SCOPE_EXIT{ if (BeaconClient) { BeaconClient->DestroyBeacon(); } };
+		UTEST_TRUE_EXPR(BeaconUnitTest::ConfigureBeacon(*Prerequisites, BeaconClient));
+
+		UTEST_TRUE_EXPR(BeaconClient->GetConnectionState() == EBeaconConnectionState::Invalid);
+		UTEST_TRUE_EXPR(BeaconUnitTest::InitClientForUser(*Prerequisites, BeaconClient, UserId));
+		UTEST_TRUE_EXPR(BeaconClient->GetConnectionState() == EBeaconConnectionState::Pending);
+		UTEST_TRUE_EXPR(BeaconClient->GetNetConnection() != nullptr);
+
+		TSharedPtr<BeaconUnitTest::FNetworkStats> BeaconClientNetStats;
+		UTEST_TRUE_EXPR(BeaconUnitTest::ConfigureBeaconNetDriver(*Prerequisites, BeaconClient, &BeaconClientNetStats));
+		UTEST_TRUE_EXPR(BeaconClientNetStats.IsValid());
+
+		// Handshake testing.
+
+		UTEST_TRUE_EXPR(BeaconUnitTest::TickUntilConnected(*Prerequisites, BeaconClient));
+		UTEST_TRUE_EXPR(BeaconClient->GetConnectionState() == EBeaconConnectionState::Open);
+		UTEST_TRUE_EXPR(Prerequisites->GetStats().Client.OnConnected.InvokeCount == 1);
+		UTEST_TRUE_EXPR(Prerequisites->GetStats().Client.OnFailure.InvokeCount == 0);
+		UTEST_TRUE_EXPR(Prerequisites->GetStats().Host.OnFailure.InvokeCount == 0);
+		UTEST_TRUE_EXPR(Prerequisites->GetStats().HostObject.OnClientConnected.InvokeCount == 1);
+		UTEST_TRUE_EXPR(Prerequisites->GetStats().HostObject.NotifyClientDisconnected.InvokeCount == 0);
+		HostUserBeaconClient = BeaconUnitTest::GetBeaconClientForUser(BeaconHostObject, UserId);
+		UTEST_TRUE_EXPR(HostUserBeaconClient != nullptr);
+		BeaconClientNetStats->ReceivedControlMessages.Empty();
+		BeaconHostNetStats->ReceivedControlMessages.Empty();
+
+		// Handshake complete, channel is open.
+
+		// Disconnect the client.
+		// The connection on the host side is now closed, but the hosts client object has not yet been notified.
+		BeaconHostObject->DisconnectClient(HostUserBeaconClient);
+		UTEST_TRUE_EXPR(BeaconClient->GetConnectionState() == EBeaconConnectionState::Open);
+		UTEST_TRUE_EXPR(Prerequisites->GetStats().Client.OnConnected.InvokeCount == 1);
+		UTEST_TRUE_EXPR(Prerequisites->GetStats().Client.OnFailure.InvokeCount == 0);
+		UTEST_TRUE_EXPR(Prerequisites->GetStats().Host.OnFailure.InvokeCount == 0);
+		UTEST_TRUE_EXPR(Prerequisites->GetStats().HostObject.OnClientConnected.InvokeCount == 1);
+		UTEST_TRUE_EXPR(Prerequisites->GetStats().HostObject.NotifyClientDisconnected.InvokeCount == 0);
+		UTEST_TRUE_EXPR(BeaconUnitTest::GetBeaconClientForUser(BeaconHostObject, UserId) != nullptr);
+
+		// Check that client and host client object cleaned up.
+		// Make sure host client object cleaned up.
+		UTEST_TRUE_EXPR(BeaconUnitTest::TickOnce(*Prerequisites));
+		UTEST_TRUE_EXPR(BeaconClient->GetConnectionState() == EBeaconConnectionState::Open);
+		UTEST_TRUE_EXPR(Prerequisites->GetStats().Client.OnConnected.InvokeCount == 1);
+		UTEST_TRUE_EXPR(Prerequisites->GetStats().Client.OnFailure.InvokeCount == 0);
+		UTEST_TRUE_EXPR(Prerequisites->GetStats().Host.OnFailure.InvokeCount == 0);
+		UTEST_TRUE_EXPR(Prerequisites->GetStats().HostObject.OnClientConnected.InvokeCount == 1);
+		UTEST_TRUE_EXPR(Prerequisites->GetStats().HostObject.NotifyClientDisconnected.InvokeCount == 1);
+		UTEST_TRUE_EXPR(BeaconUnitTest::GetBeaconClientForUser(BeaconHostObject, UserId) == nullptr);
+
+		// Make sure client cleaned up.
+		UTEST_TRUE_EXPR(BeaconUnitTest::TickOnce(*Prerequisites));
+		UTEST_TRUE_EXPR(BeaconClient->GetConnectionState() == EBeaconConnectionState::Invalid);
+		UTEST_TRUE_EXPR(Prerequisites->GetStats().Client.OnConnected.InvokeCount == 1);
+		UTEST_TRUE_EXPR(Prerequisites->GetStats().Client.OnFailure.InvokeCount == 1);
+		UTEST_TRUE_EXPR(Prerequisites->GetStats().Host.OnFailure.InvokeCount == 0);
+		UTEST_TRUE_EXPR(Prerequisites->GetStats().HostObject.OnClientConnected.InvokeCount == 1);
+		UTEST_TRUE_EXPR(Prerequisites->GetStats().HostObject.NotifyClientDisconnected.InvokeCount == 1);
+		UTEST_TRUE_EXPR(BeaconUnitTest::GetBeaconClientForUser(BeaconHostObject, UserId) == nullptr);
+
+		// Check that encryption delegates were not fired.
+		UTEST_TRUE_EXPR(Prerequisites->GetStats().Encryption.NetworkEncryptionToken.InvokeCount == 0);
+		UTEST_TRUE_EXPR(Prerequisites->GetStats().Encryption.NetworkEncryptionToken.CallbackCount == 0);
+		UTEST_TRUE_EXPR(Prerequisites->GetStats().Encryption.NetworkEncryptionAck.InvokeCount == 0);
+		UTEST_TRUE_EXPR(Prerequisites->GetStats().Encryption.NetworkEncryptionAck.CallbackCount == 0);
+		// NetworkEncryptionFailure is always fired on client / host when no encryption data is provided.
+		UTEST_TRUE_EXPR(Prerequisites->GetStats().Encryption.NetworkEncryptionFailure.InvokeCount == 2);
+	}
+	
+	return true;
+}
+
+IMPLEMENT_SIMPLE_AUTOMATION_TEST(FOnlineBeaconTestBasicHandshakeClientTimeout,
+	"System.Engine.Online.OnlineSubsystemUtils.OnlineBeacon.BasicHandshakeClientTimeout",
+	EAutomationTestFlags::EditorContext | EAutomationTestFlags::EngineFilter)
+bool FOnlineBeaconTestBasicHandshakeClientTimeout::RunTest(const FString& Parameters)
+{
+	TSharedPtr<BeaconUnitTest::FTestPrerequisites> Prerequisites = BeaconUnitTest::FTestPrerequisites::TryCreate();
+	UTEST_TRUE_EXPR(Prerequisites.IsValid());
+
+	Prerequisites->GetConfig().NetDriver.ServerListenPort = 9999;
+
+	// Don't allow the host to timeout.
+	Prerequisites->GetConfig().NetDriver.bNoTimeouts = true;
+
+	// Snapshot config to restore before running a test section.
+	const BeaconUnitTest::FTestConfig BaseConfig = Prerequisites->GetConfig();
+
+	FUniqueNetIdStringRef UserId = FUniqueNetIdString::Create(TEXT("User"), TEXT("UnitTest"));
+
+	// Host setup.
+	TObjectPtr<AOnlineBeaconUnitTestHostObject> BeaconHostObject = Prerequisites->GetWorld()->SpawnActor<AOnlineBeaconUnitTestHostObject>(AOnlineBeaconUnitTestHostObject::StaticClass());
+	TObjectPtr<AOnlineBeaconUnitTestHost> BeaconHost = Prerequisites->GetWorld()->SpawnActor<AOnlineBeaconUnitTestHost>(AOnlineBeaconUnitTestHost::StaticClass());
+	ON_SCOPE_EXIT{ if (BeaconHostObject) { BeaconHostObject->Destroy(); } };
+	ON_SCOPE_EXIT{ if (BeaconHost) { BeaconHost->DestroyBeacon(); } };
+	UTEST_TRUE_EXPR(BeaconUnitTest::ConfigureBeacon(*Prerequisites, BeaconHost));
+	BeaconHost->ListenPort = Prerequisites->GetConfig().NetDriver.ServerListenPort;
+	BeaconHost->RegisterHost(BeaconHostObject);
+	UTEST_TRUE_EXPR(BeaconHost->InitHost());
+	BeaconHost->PauseBeaconRequests(false);
+
+	TSharedPtr<BeaconUnitTest::FNetworkStats> BeaconHostNetStats;
+	UTEST_TRUE_EXPR(BeaconUnitTest::ConfigureBeaconNetDriver(*Prerequisites, BeaconHost, &BeaconHostNetStats));
+	UTEST_TRUE_EXPR(BeaconHostNetStats.IsValid());
+
+	// Timeout after client sends NMT_Hello
+	{
+		BeaconHostNetStats->ReceivedControlMessages.Empty();
+		Prerequisites->GetConfig() = BaseConfig;
+		Prerequisites->GetStats() = BeaconUnitTest::FTestStats();
+
+		Prerequisites->GetConfig().NetDriver.bNoTimeouts = false;
+
+		// Client setup.
+
+		TObjectPtr<AOnlineBeaconUnitTestClient> BeaconClient = Prerequisites->GetWorld()->SpawnActor<AOnlineBeaconUnitTestClient>(AOnlineBeaconUnitTestClient::StaticClass());
+		ON_SCOPE_EXIT{ if (BeaconClient) { BeaconClient->DestroyBeacon(); } };
+		UTEST_TRUE_EXPR(BeaconUnitTest::ConfigureBeacon(*Prerequisites, BeaconClient));
+
+		UTEST_TRUE_EXPR(BeaconClient->GetConnectionState() == EBeaconConnectionState::Invalid);
+		UTEST_TRUE_EXPR(BeaconUnitTest::InitClientForUser(*Prerequisites, BeaconClient, UserId));
+		UTEST_TRUE_EXPR(BeaconClient->GetConnectionState() == EBeaconConnectionState::Pending);
+		UTEST_TRUE_EXPR(BeaconClient->GetNetConnection() != nullptr);
+
+		TSharedPtr<BeaconUnitTest::FNetworkStats> BeaconClientNetStats;
+		UTEST_TRUE_EXPR(BeaconUnitTest::ConfigureBeaconNetDriver(*Prerequisites, BeaconClient, &BeaconClientNetStats));
+		UTEST_TRUE_EXPR(BeaconClientNetStats.IsValid());
+
+		// Handshake testing.
+
+		UTEST_TRUE_EXPR(BeaconUnitTest::TickUntilControlMessageReceived(*Prerequisites, *BeaconHostNetStats));
+		UTEST_TRUE_EXPR(BeaconClient->GetConnectionState() == EBeaconConnectionState::Pending);
+		UTEST_TRUE_EXPR(BeaconClientNetStats->ReceivedControlMessages.IsEmpty());
+		UTEST_TRUE_EXPR(BeaconHostNetStats->ReceivedControlMessages == TArray<uint8>({ NMT_Hello }));
+		BeaconHostNetStats->ReceivedControlMessages.Empty();
+
+		// Block the client from sending packets.
+		UTEST_TRUE_EXPR(BeaconUnitTest::SetSocketFlags(BeaconClient, BeaconUnitTest::ESocketFlags::RecvEnabled));
+
+		UTEST_TRUE_EXPR(BeaconUnitTest::TickUntilDisconnected(*Prerequisites, BeaconClient, BeaconUnitTest::ETickFlags::SleepTickTime));
+		UTEST_TRUE_EXPR(BeaconClient->GetConnectionState() == EBeaconConnectionState::Invalid);
+		UTEST_TRUE_EXPR(BeaconClientNetStats->ReceivedControlMessages == TArray<uint8>({ NMT_BeaconWelcome }));
+		UTEST_TRUE_EXPR(BeaconHostNetStats->ReceivedControlMessages.IsEmpty());
+	}
+
+	// Timeout after client sends NMT_BeaconJoin
+	{
+		BeaconHostNetStats->ReceivedControlMessages.Empty();
+		Prerequisites->GetConfig() = BaseConfig;
+		Prerequisites->GetStats() = BeaconUnitTest::FTestStats();
+
+		Prerequisites->GetConfig().NetDriver.bNoTimeouts = false;
+
+		// Client setup.
+
+		TObjectPtr<AOnlineBeaconUnitTestClient> BeaconClient = Prerequisites->GetWorld()->SpawnActor<AOnlineBeaconUnitTestClient>(AOnlineBeaconUnitTestClient::StaticClass());
+		ON_SCOPE_EXIT{ if (BeaconClient) { BeaconClient->DestroyBeacon(); } };
+		UTEST_TRUE_EXPR(BeaconUnitTest::ConfigureBeacon(*Prerequisites, BeaconClient));
+
+		UTEST_TRUE_EXPR(BeaconClient->GetConnectionState() == EBeaconConnectionState::Invalid);
+		UTEST_TRUE_EXPR(BeaconUnitTest::InitClientForUser(*Prerequisites, BeaconClient, UserId));
+		UTEST_TRUE_EXPR(BeaconClient->GetConnectionState() == EBeaconConnectionState::Pending);
+		UTEST_TRUE_EXPR(BeaconClient->GetNetConnection() != nullptr);
+
+		TSharedPtr<BeaconUnitTest::FNetworkStats> BeaconClientNetStats;
+		UTEST_TRUE_EXPR(BeaconUnitTest::ConfigureBeaconNetDriver(*Prerequisites, BeaconClient, &BeaconClientNetStats));
+		UTEST_TRUE_EXPR(BeaconClientNetStats.IsValid());
+
+		// Handshake testing.
+
+		UTEST_TRUE_EXPR(BeaconUnitTest::TickUntilControlMessageReceived(*Prerequisites, *BeaconHostNetStats));
+		UTEST_TRUE_EXPR(BeaconClient->GetConnectionState() == EBeaconConnectionState::Pending);
+		UTEST_TRUE_EXPR(BeaconClientNetStats->ReceivedControlMessages.IsEmpty());
+		UTEST_TRUE_EXPR(BeaconHostNetStats->ReceivedControlMessages == TArray<uint8>({ NMT_Hello }));
+		BeaconHostNetStats->ReceivedControlMessages.Empty();
+
+		UTEST_TRUE_EXPR(BeaconUnitTest::TickOnce(*Prerequisites));
+		UTEST_TRUE_EXPR(BeaconClient->GetConnectionState() == EBeaconConnectionState::Pending);
+		UTEST_TRUE_EXPR(BeaconClientNetStats->ReceivedControlMessages == TArray<uint8>({ NMT_BeaconWelcome }));
+		UTEST_TRUE_EXPR(BeaconHostNetStats->ReceivedControlMessages.IsEmpty());
+		BeaconClientNetStats->ReceivedControlMessages.Empty();
+
+		UTEST_TRUE_EXPR(BeaconUnitTest::TickOnce(*Prerequisites));
+		UTEST_TRUE_EXPR(BeaconClient->GetConnectionState() == EBeaconConnectionState::Pending);
+		UTEST_TRUE_EXPR(BeaconHostNetStats->ReceivedControlMessages == TArray<uint8>({ NMT_Netspeed, NMT_BeaconJoin }));
+		BeaconHostNetStats->ReceivedControlMessages.Empty();
+
+		// Block the client from sending packets.
+		UTEST_TRUE_EXPR(BeaconUnitTest::SetSocketFlags(BeaconClient, BeaconUnitTest::ESocketFlags::RecvEnabled));
+
+		UTEST_TRUE_EXPR(BeaconUnitTest::TickUntilDisconnected(*Prerequisites, BeaconClient, BeaconUnitTest::ETickFlags::SleepTickTime));
+		UTEST_TRUE_EXPR(BeaconClient->GetConnectionState() == EBeaconConnectionState::Invalid);
+		UTEST_TRUE_EXPR(BeaconClientNetStats->ReceivedControlMessages == TArray<uint8>({ NMT_BeaconAssignGUID }));
+		UTEST_TRUE_EXPR(BeaconHostNetStats->ReceivedControlMessages.IsEmpty());
+	}
+
+	// Timeout after client sends NMT_BeaconNetGUIDAck
+	{
+		BeaconHostNetStats->ReceivedControlMessages.Empty();
+		Prerequisites->GetConfig() = BaseConfig;
+		Prerequisites->GetStats() = BeaconUnitTest::FTestStats();
+
+		Prerequisites->GetConfig().NetDriver.bNoTimeouts = false;
+
+		// Client setup.
+
+		TObjectPtr<AOnlineBeaconUnitTestClient> BeaconClient = Prerequisites->GetWorld()->SpawnActor<AOnlineBeaconUnitTestClient>(AOnlineBeaconUnitTestClient::StaticClass());
+		ON_SCOPE_EXIT{ if (BeaconClient) { BeaconClient->DestroyBeacon(); } };
+		UTEST_TRUE_EXPR(BeaconUnitTest::ConfigureBeacon(*Prerequisites, BeaconClient));
+
+		UTEST_TRUE_EXPR(BeaconClient->GetConnectionState() == EBeaconConnectionState::Invalid);
+		UTEST_TRUE_EXPR(BeaconUnitTest::InitClientForUser(*Prerequisites, BeaconClient, UserId));
+		UTEST_TRUE_EXPR(BeaconClient->GetConnectionState() == EBeaconConnectionState::Pending);
+		UTEST_TRUE_EXPR(BeaconClient->GetNetConnection() != nullptr);
+
+		TSharedPtr<BeaconUnitTest::FNetworkStats> BeaconClientNetStats;
+		UTEST_TRUE_EXPR(BeaconUnitTest::ConfigureBeaconNetDriver(*Prerequisites, BeaconClient, &BeaconClientNetStats));
+		UTEST_TRUE_EXPR(BeaconClientNetStats.IsValid());
+
+		// Handshake testing.
+
+		UTEST_TRUE_EXPR(BeaconUnitTest::TickUntilControlMessageReceived(*Prerequisites, *BeaconHostNetStats));
+		UTEST_TRUE_EXPR(BeaconClientNetStats->ReceivedControlMessages.IsEmpty());
+		UTEST_TRUE_EXPR(BeaconHostNetStats->ReceivedControlMessages == TArray<uint8>({ NMT_Hello }));
+		BeaconHostNetStats->ReceivedControlMessages.Empty();
+
+		UTEST_TRUE_EXPR(BeaconUnitTest::TickOnce(*Prerequisites));
+		UTEST_TRUE_EXPR(BeaconClient->GetConnectionState() == EBeaconConnectionState::Pending);
+		UTEST_TRUE_EXPR(BeaconClientNetStats->ReceivedControlMessages == TArray<uint8>({ NMT_BeaconWelcome }));
+		UTEST_TRUE_EXPR(BeaconHostNetStats->ReceivedControlMessages.IsEmpty());
+		BeaconClientNetStats->ReceivedControlMessages.Empty();
+
+		UTEST_TRUE_EXPR(BeaconUnitTest::TickOnce(*Prerequisites));
+		UTEST_TRUE_EXPR(BeaconClient->GetConnectionState() == EBeaconConnectionState::Pending);
+		UTEST_TRUE_EXPR(BeaconClientNetStats->ReceivedControlMessages.IsEmpty());
+		UTEST_TRUE_EXPR(BeaconHostNetStats->ReceivedControlMessages == TArray<uint8>({ NMT_Netspeed, NMT_BeaconJoin }));
+		BeaconHostNetStats->ReceivedControlMessages.Empty();
+
+		UTEST_TRUE_EXPR(BeaconUnitTest::TickOnce(*Prerequisites));
+		UTEST_TRUE_EXPR(BeaconClient->GetConnectionState() == EBeaconConnectionState::Pending);
+		UTEST_TRUE_EXPR(BeaconClientNetStats->ReceivedControlMessages == TArray<uint8>({ NMT_BeaconAssignGUID }));
+		UTEST_TRUE_EXPR(BeaconHostNetStats->ReceivedControlMessages.IsEmpty());
+		BeaconClientNetStats->ReceivedControlMessages.Empty();
+
+		UTEST_TRUE_EXPR(BeaconUnitTest::TickOnce(*Prerequisites));
+		UTEST_TRUE_EXPR(BeaconClient->GetConnectionState() == EBeaconConnectionState::Pending);
+		UTEST_TRUE_EXPR(BeaconClientNetStats->ReceivedControlMessages.IsEmpty());
+		UTEST_TRUE_EXPR(BeaconHostNetStats->ReceivedControlMessages == TArray<uint8>({ NMT_BeaconNetGUIDAck }));
+		BeaconHostNetStats->ReceivedControlMessages.Empty();
+
+		// Block the client from sending packets.
+		UTEST_TRUE_EXPR(BeaconUnitTest::SetSocketFlags(BeaconClient, BeaconUnitTest::ESocketFlags::RecvEnabled));
+
+		UTEST_TRUE_EXPR(BeaconUnitTest::TickUntilDisconnected(*Prerequisites, BeaconClient, BeaconUnitTest::ETickFlags::SleepTickTime));
+		UTEST_TRUE_EXPR(BeaconClient->GetConnectionState() == EBeaconConnectionState::Invalid);
+		UTEST_TRUE_EXPR(BeaconClientNetStats->ReceivedControlMessages.IsEmpty());
+		UTEST_TRUE_EXPR(BeaconHostNetStats->ReceivedControlMessages.IsEmpty());
+	}
+
+	return true;
+}
+
+IMPLEMENT_SIMPLE_AUTOMATION_TEST(FOnlineBeaconTestBasicHandshakeHostTimeout,
+	"System.Engine.Online.OnlineSubsystemUtils.OnlineBeacon.BasicHandshakeHostTimeout",
+	EAutomationTestFlags::EditorContext | EAutomationTestFlags::EngineFilter)
+bool FOnlineBeaconTestBasicHandshakeHostTimeout::RunTest(const FString& Parameters)
+{
+	TSharedPtr<BeaconUnitTest::FTestPrerequisites> Prerequisites = BeaconUnitTest::FTestPrerequisites::TryCreate();
+	UTEST_TRUE_EXPR(Prerequisites.IsValid());
+
+	Prerequisites->GetConfig().NetDriver.ServerListenPort = 9999;
+
+	// Snapshot config to restore before running a test section.
+	const BeaconUnitTest::FTestConfig BaseConfig = Prerequisites->GetConfig();
+
+	FUniqueNetIdStringRef UserId = FUniqueNetIdString::Create(TEXT("User"), TEXT("UnitTest"));
+
+	// Host setup.
+	TObjectPtr<AOnlineBeaconUnitTestHostObject> BeaconHostObject = Prerequisites->GetWorld()->SpawnActor<AOnlineBeaconUnitTestHostObject>(AOnlineBeaconUnitTestHostObject::StaticClass());
+	TObjectPtr<AOnlineBeaconUnitTestHost> BeaconHost = Prerequisites->GetWorld()->SpawnActor<AOnlineBeaconUnitTestHost>(AOnlineBeaconUnitTestHost::StaticClass());
+	ON_SCOPE_EXIT{ if (BeaconHostObject) { BeaconHostObject->Destroy(); } };
+	ON_SCOPE_EXIT{ if (BeaconHost) { BeaconHost->DestroyBeacon(); } };
+	UTEST_TRUE_EXPR(BeaconUnitTest::ConfigureBeacon(*Prerequisites, BeaconHost));
+	BeaconHost->ListenPort = Prerequisites->GetConfig().NetDriver.ServerListenPort;
+	BeaconHost->RegisterHost(BeaconHostObject);
+	UTEST_TRUE_EXPR(BeaconHost->InitHost());
+	BeaconHost->PauseBeaconRequests(false);
+
+	TSharedPtr<BeaconUnitTest::FNetworkStats> BeaconHostNetStats;
+	UTEST_TRUE_EXPR(BeaconUnitTest::ConfigureBeaconNetDriver(*Prerequisites, BeaconHost, &BeaconHostNetStats));
+	UTEST_TRUE_EXPR(BeaconHostNetStats.IsValid());
+
+	// Timeout after host sends NMT_BeaconWelcome
+	{
+		BeaconHostNetStats->ReceivedControlMessages.Empty();
+		Prerequisites->GetConfig() = BaseConfig;
+		Prerequisites->GetStats() = BeaconUnitTest::FTestStats();
+
+		Prerequisites->GetConfig().NetDriver.bNoTimeouts = true;
+
+		// Client setup.
+
+		TObjectPtr<AOnlineBeaconUnitTestClient> BeaconClient = Prerequisites->GetWorld()->SpawnActor<AOnlineBeaconUnitTestClient>(AOnlineBeaconUnitTestClient::StaticClass());
+		ON_SCOPE_EXIT{ if (BeaconClient) { BeaconClient->DestroyBeacon(); } };
+		UTEST_TRUE_EXPR(BeaconUnitTest::ConfigureBeacon(*Prerequisites, BeaconClient));
+
+		UTEST_TRUE_EXPR(BeaconClient->GetConnectionState() == EBeaconConnectionState::Invalid);
+		UTEST_TRUE_EXPR(BeaconUnitTest::InitClientForUser(*Prerequisites, BeaconClient, UserId));
+		UTEST_TRUE_EXPR(BeaconClient->GetConnectionState() == EBeaconConnectionState::Pending);
+		UTEST_TRUE_EXPR(BeaconClient->GetNetConnection() != nullptr);
+
+		TSharedPtr<BeaconUnitTest::FNetworkStats> BeaconClientNetStats;
+		UTEST_TRUE_EXPR(BeaconUnitTest::ConfigureBeaconNetDriver(*Prerequisites, BeaconClient, &BeaconClientNetStats));
+		UTEST_TRUE_EXPR(BeaconClientNetStats.IsValid());
+
+		// Handshake testing.
+
+		UTEST_TRUE_EXPR(BeaconUnitTest::TickUntilControlMessageReceived(*Prerequisites, *BeaconHostNetStats));
+		UTEST_TRUE_EXPR(BeaconClient->GetConnectionState() == EBeaconConnectionState::Pending);
+		UTEST_TRUE_EXPR(BeaconClientNetStats->ReceivedControlMessages.IsEmpty());
+		UTEST_TRUE_EXPR(BeaconHostNetStats->ReceivedControlMessages == TArray<uint8>({ NMT_Hello }));
+		BeaconHostNetStats->ReceivedControlMessages.Empty();
+
+		// Block the client from handling packets.
+		UTEST_TRUE_EXPR(BeaconUnitTest::SetSocketFlags(BeaconClient, BeaconUnitTest::ESocketFlags::Disabled));
+
+		UTEST_TRUE_EXPR(BeaconUnitTest::TickUntilTimeoutElapsed(*Prerequisites, BeaconUnitTest::ETickFlags::SleepTickTime));
+
+		// Unblock the client from handling packets.
+		UTEST_TRUE_EXPR(BeaconUnitTest::SetSocketFlags(BeaconClient, BeaconUnitTest::ESocketFlags::Default));
+
+		// Try to continue the handshake after the host has cleaned up the client state.
+		// The host will not see the clients control message since it has closed the connection.
+		// The client will be in an invalid state due to receiving the close packet.
+
+		UTEST_TRUE_EXPR(BeaconUnitTest::TickOnce(*Prerequisites));
+		UTEST_TRUE_EXPR(BeaconClient->GetConnectionState() == EBeaconConnectionState::Invalid);
+		UTEST_TRUE_EXPR(BeaconClientNetStats->ReceivedControlMessages == TArray<uint8>({ NMT_BeaconWelcome }));
+		UTEST_TRUE_EXPR(BeaconHostNetStats->ReceivedControlMessages.IsEmpty());
+	}
+
+	// Timeout after host sends NMT_BeaconAssignGUID
+	{
+		BeaconHostNetStats->ReceivedControlMessages.Empty();
+		Prerequisites->GetConfig() = BaseConfig;
+		Prerequisites->GetStats() = BeaconUnitTest::FTestStats();
+
+		Prerequisites->GetConfig().NetDriver.bNoTimeouts = true;
+
+		// Client setup.
+
+		TObjectPtr<AOnlineBeaconUnitTestClient> BeaconClient = Prerequisites->GetWorld()->SpawnActor<AOnlineBeaconUnitTestClient>(AOnlineBeaconUnitTestClient::StaticClass());
+		ON_SCOPE_EXIT{ if (BeaconClient) { BeaconClient->DestroyBeacon(); } };
+		UTEST_TRUE_EXPR(BeaconUnitTest::ConfigureBeacon(*Prerequisites, BeaconClient));
+
+		UTEST_TRUE_EXPR(BeaconClient->GetConnectionState() == EBeaconConnectionState::Invalid);
+		UTEST_TRUE_EXPR(BeaconUnitTest::InitClientForUser(*Prerequisites, BeaconClient, UserId));
+		UTEST_TRUE_EXPR(BeaconClient->GetConnectionState() == EBeaconConnectionState::Pending);
+		UTEST_TRUE_EXPR(BeaconClient->GetNetConnection() != nullptr);
+
+		TSharedPtr<BeaconUnitTest::FNetworkStats> BeaconClientNetStats;
+		UTEST_TRUE_EXPR(BeaconUnitTest::ConfigureBeaconNetDriver(*Prerequisites, BeaconClient, &BeaconClientNetStats));
+		UTEST_TRUE_EXPR(BeaconClientNetStats.IsValid());
+
+		// Handshake testing.
+
+		UTEST_TRUE_EXPR(BeaconUnitTest::TickUntilControlMessageReceived(*Prerequisites, *BeaconHostNetStats));
+		UTEST_TRUE_EXPR(BeaconClient->GetConnectionState() == EBeaconConnectionState::Pending);
+		UTEST_TRUE_EXPR(BeaconClientNetStats->ReceivedControlMessages.IsEmpty());
+		UTEST_TRUE_EXPR(BeaconHostNetStats->ReceivedControlMessages == TArray<uint8>({ NMT_Hello }));
+		BeaconHostNetStats->ReceivedControlMessages.Empty();
+
+		UTEST_TRUE_EXPR(BeaconUnitTest::TickOnce(*Prerequisites));
+		UTEST_TRUE_EXPR(BeaconClient->GetConnectionState() == EBeaconConnectionState::Pending);
+		UTEST_TRUE_EXPR(BeaconClientNetStats->ReceivedControlMessages == TArray<uint8>({ NMT_BeaconWelcome }));
+		UTEST_TRUE_EXPR(BeaconHostNetStats->ReceivedControlMessages.IsEmpty());
+		BeaconClientNetStats->ReceivedControlMessages.Empty();
+
+		UTEST_TRUE_EXPR(BeaconUnitTest::TickOnce(*Prerequisites));
+		UTEST_TRUE_EXPR(BeaconClient->GetConnectionState() == EBeaconConnectionState::Pending);
+		UTEST_TRUE_EXPR(BeaconClientNetStats->ReceivedControlMessages.IsEmpty());
+		UTEST_TRUE_EXPR(BeaconHostNetStats->ReceivedControlMessages == TArray<uint8>({ NMT_Netspeed, NMT_BeaconJoin }));
+		BeaconHostNetStats->ReceivedControlMessages.Empty();
+
+		// Block the client from handling packets.
+		UTEST_TRUE_EXPR(BeaconUnitTest::SetSocketFlags(BeaconClient, BeaconUnitTest::ESocketFlags::Disabled));
+
+		UTEST_TRUE_EXPR(BeaconUnitTest::TickUntilTimeoutElapsed(*Prerequisites, BeaconUnitTest::ETickFlags::SleepTickTime));
+
+		// Unblock the client from handling packets.
+		UTEST_TRUE_EXPR(BeaconUnitTest::SetSocketFlags(BeaconClient, BeaconUnitTest::ESocketFlags::Default));
+
+		// Try to continue the handshake after the host has cleaned up the client state.
+		// The host will not see the clients control message since it has closed the connection.
+		// The client will be in an invalid state due to receiving the close packet.
+
+		UTEST_TRUE_EXPR(BeaconUnitTest::TickOnce(*Prerequisites));
+		UTEST_TRUE_EXPR(BeaconClient->GetConnectionState() == EBeaconConnectionState::Invalid);
+		UTEST_TRUE_EXPR(BeaconClientNetStats->ReceivedControlMessages == TArray<uint8>({ NMT_BeaconAssignGUID }));
+		UTEST_TRUE_EXPR(BeaconHostNetStats->ReceivedControlMessages.IsEmpty());
+	}
+
+	return true;
+}
+
+IMPLEMENT_SIMPLE_AUTOMATION_TEST(FOnlineBeaconTestEncryptedHandshakeSuccess,
+	"System.Engine.Online.OnlineSubsystemUtils.OnlineBeacon.EncryptedHandshakeSuccess",
+	EAutomationTestFlags::EditorContext | EAutomationTestFlags::EngineFilter)
+bool FOnlineBeaconTestEncryptedHandshakeSuccess::RunTest(const FString& Parameters)
+{
+	TSharedPtr<BeaconUnitTest::FTestPrerequisites> Prerequisites = BeaconUnitTest::FTestPrerequisites::TryCreate();
+	UTEST_TRUE_EXPR(Prerequisites.IsValid());
+
+	// Check that encryption is allowed.
+	UTEST_TRUE_EXPR(CVarNetAllowEncryption.GetValueOnGameThread() != 0);
+
+	Prerequisites->GetConfig().NetDriver.bNoTimeouts = true;
+	Prerequisites->GetConfig().NetDriver.ServerListenPort = 9999;
+
+	// Todo: Remove external dependency on AESGCMHandlerComponent.
+	Prerequisites->GetConfig().Encryption.bEnabled = true;
+	Prerequisites->GetConfig().Encryption.NetDriverEncryptionComponentName = TEXT("AESGCMHandlerComponent");
+
+	// Snapshot config to restore before running a test section.
+	const BeaconUnitTest::FTestConfig BaseConfig = Prerequisites->GetConfig();
+
+	// Valid encryption key.
+	const FString EncryptionIdentitfier = TEXT("test");
+	const FString Base64EncryptionKey = TEXT("IYaVIE38d6J9VfbPULuSMfn3/axig797U8DVJyRm1/c=");
+	TArray<uint8> DecodedEncryptionKey;
+	UTEST_TRUE_EXPR(FBase64::Decode(Base64EncryptionKey, DecodedEncryptionKey));
+
+	FUniqueNetIdStringRef UserId = FUniqueNetIdString::Create(TEXT("User"), TEXT("UnitTest"));
+
+	// Host setup.
+	TObjectPtr<AOnlineBeaconUnitTestHostObject> BeaconHostObject = Prerequisites->GetWorld()->SpawnActor<AOnlineBeaconUnitTestHostObject>(AOnlineBeaconUnitTestHostObject::StaticClass());
+	TObjectPtr<AOnlineBeaconUnitTestHost> BeaconHost = Prerequisites->GetWorld()->SpawnActor<AOnlineBeaconUnitTestHost>(AOnlineBeaconUnitTestHost::StaticClass());
+	ON_SCOPE_EXIT{ if (BeaconHostObject) { BeaconHostObject->Destroy(); } };
+	ON_SCOPE_EXIT{ if (BeaconHost) { BeaconHost->DestroyBeacon(); } };
+	UTEST_TRUE_EXPR(BeaconUnitTest::ConfigureBeacon(*Prerequisites, BeaconHost));
+	BeaconHost->ListenPort = Prerequisites->GetConfig().NetDriver.ServerListenPort;
+	BeaconHost->RegisterHost(BeaconHostObject);
+	UTEST_TRUE_EXPR(BeaconHost->InitHost());
+	BeaconHost->PauseBeaconRequests(false);
+
+	TSharedPtr<BeaconUnitTest::FNetworkStats> BeaconHostNetStats;
+	UTEST_TRUE_EXPR(BeaconUnitTest::ConfigureBeaconNetDriver(*Prerequisites, BeaconHost, &BeaconHostNetStats));
+	UTEST_TRUE_EXPR(BeaconHostNetStats.IsValid());
+
+	// Successful handshake.
+	// OnReceivedNetworkEncryptionToken and OnReceivedNetworkEncryptionAck respond immediately.
+	{
+		BeaconHostNetStats->ReceivedControlMessages.Empty();
+		Prerequisites->GetConfig() = BaseConfig;
+		Prerequisites->GetStats() = BeaconUnitTest::FTestStats();
+
+		// Setup encryption.
+		// 
+		// Future work: The beacon encryption handshake doesn't currently work as the encryption structure intends.
+		// 
+		// Intended flow: The encryption data is fetched using net delegate
+		// Notes:
+		//  The encryption data is fetched using net delegate
+		// Sequence:
+		//  1. Client sends encryption identifier.
+		//  2. Host finds the key for the identifier and sends NMT_EncryptionAck, NMT_BeaconWelcome. - calls FNetDelegates::OnReceivedNetworkEncryptionToken
+		//  3. Client finds key for identifier and enables encryption. - calls FNetDelegates::OnReceivedNetworkEncryptionAck.
+		//
+		// Current flow:
+		// Notes:
+		//  The encryption data is stored on the OnlineBeaconClient class
+		//  The intended sequence can be forced by clearing the encryption data on the OnlineBeaconClient object after sending the initial hello packet.
+		// Sequence:
+		//  1. Client sets encryption key on the connection.
+		//  2. Client sends encryption identifier.
+		//  3. Host finds the key for the identifier and sends NMT_EncryptionAck, NMT_BeaconWelcome. - calls FNetDelegates::OnReceivedNetworkEncryptionToken
+		//  4. Client uses stored key for identifier and enables encryption. - FNetDelegates::OnReceivedNetworkEncryptionAck is bypassed.
+
+		Prerequisites->GetConfig().Encryption.FailureAction = EEncryptionFailureAction::RejectConnection;
+
+		// Host config.
+		Prerequisites->GetConfig().Encryption.Host.bDelayDelegate = false;
+		Prerequisites->GetConfig().Encryption.Host.Response = EEncryptionResponse::Success;
+		Prerequisites->GetConfig().Encryption.Host.ErrorMsg.Empty();
+		Prerequisites->GetConfig().Encryption.Host.EncryptionData.Key = DecodedEncryptionKey;
+		Prerequisites->GetConfig().Encryption.Host.EncryptionData.Identifier = EncryptionIdentitfier;
+
+		// Client config.
+		Prerequisites->GetConfig().Encryption.Client.bDelayDelegate = false;
+		Prerequisites->GetConfig().Encryption.Client.Response = EEncryptionResponse::Success;
+		Prerequisites->GetConfig().Encryption.Client.ErrorMsg.Empty();
+		Prerequisites->GetConfig().Encryption.Client.EncryptionData.Key = DecodedEncryptionKey;
+		Prerequisites->GetConfig().Encryption.Client.EncryptionData.Identifier = EncryptionIdentitfier;
+
+		Prerequisites->GetStats() = BeaconUnitTest::FTestStats();
+
+		TObjectPtr<AOnlineBeaconClient> HostUserBeaconClient;
+
+		// Client setup.
+
+		TObjectPtr<AOnlineBeaconUnitTestClient> BeaconClient = Prerequisites->GetWorld()->SpawnActor<AOnlineBeaconUnitTestClient>(AOnlineBeaconUnitTestClient::StaticClass());
+		ON_SCOPE_EXIT{ if (BeaconClient) { BeaconClient->DestroyBeacon(); } };
+		UTEST_TRUE_EXPR(BeaconUnitTest::ConfigureBeacon(*Prerequisites, BeaconClient));
+
+		BeaconClient->SetEncryptionData(Prerequisites->GetConfig().Encryption.Client.EncryptionData);
+		UTEST_TRUE_EXPR(BeaconUnitTest::InitClientForUser(*Prerequisites, BeaconClient, UserId));
+		UTEST_TRUE_EXPR(BeaconClient->GetConnectionState() == EBeaconConnectionState::Pending);
+		UTEST_TRUE_EXPR(Prerequisites->GetStats().Encryption.NetworkEncryptionFailure.InvokeCount == 0);
+
+		TSharedPtr<BeaconUnitTest::FNetworkStats> BeaconClientNetStats;
+		UTEST_TRUE_EXPR(BeaconUnitTest::ConfigureBeaconNetDriver(*Prerequisites, BeaconClient, &BeaconClientNetStats));
+		UTEST_TRUE_EXPR(BeaconClientNetStats.IsValid());
+
+		// Handshake testing.
+
+		UTEST_TRUE_EXPR(BeaconUnitTest::TickUntilControlMessageReceived(*Prerequisites, *BeaconHostNetStats));
+		UTEST_TRUE_EXPR(BeaconClient->GetConnectionState() == EBeaconConnectionState::Pending);
+		UTEST_TRUE_EXPR(BeaconClientNetStats->ReceivedControlMessages.IsEmpty());
+		UTEST_TRUE_EXPR(BeaconHostNetStats->ReceivedControlMessages == TArray<uint8>({ NMT_Hello }));
+		UTEST_TRUE_EXPR(Prerequisites->GetStats().Encryption.NetworkEncryptionToken.InvokeCount == 1);
+		UTEST_TRUE_EXPR(Prerequisites->GetStats().Encryption.NetworkEncryptionFailure.InvokeCount == 0);
+		BeaconHostNetStats->ReceivedControlMessages.Empty();
+
+		UTEST_TRUE_EXPR(BeaconUnitTest::TickOnce(*Prerequisites));
+		UTEST_TRUE_EXPR(BeaconClient->GetConnectionState() == EBeaconConnectionState::Pending);
+		UTEST_TRUE_EXPR(BeaconClientNetStats->ReceivedControlMessages == TArray<uint8>({ NMT_EncryptionAck, NMT_BeaconWelcome }));
+		UTEST_TRUE_EXPR(BeaconHostNetStats->ReceivedControlMessages.IsEmpty());
+		UTEST_TRUE_EXPR(Prerequisites->GetStats().Encryption.NetworkEncryptionAck.InvokeCount == 0);
+		UTEST_TRUE_EXPR(Prerequisites->GetStats().Encryption.NetworkEncryptionFailure.InvokeCount == 0);
+		BeaconClientNetStats->ReceivedControlMessages.Empty();
+
+		UTEST_TRUE_EXPR(BeaconUnitTest::TickOnce(*Prerequisites));
+		UTEST_TRUE_EXPR(BeaconClient->GetConnectionState() == EBeaconConnectionState::Pending);
+		UTEST_TRUE_EXPR(BeaconClientNetStats->ReceivedControlMessages.IsEmpty());
+		UTEST_TRUE_EXPR(BeaconHostNetStats->ReceivedControlMessages == TArray<uint8>({ NMT_Netspeed, NMT_BeaconJoin }));
+		BeaconHostNetStats->ReceivedControlMessages.Empty();
+
+		UTEST_TRUE_EXPR(BeaconUnitTest::TickOnce(*Prerequisites));
+		UTEST_TRUE_EXPR(BeaconClient->GetConnectionState() == EBeaconConnectionState::Pending);
+		UTEST_TRUE_EXPR(BeaconClientNetStats->ReceivedControlMessages == TArray<uint8>({ NMT_BeaconAssignGUID }));
+		UTEST_TRUE_EXPR(BeaconHostNetStats->ReceivedControlMessages.IsEmpty());
+		BeaconClientNetStats->ReceivedControlMessages.Empty();
+
+		UTEST_TRUE_EXPR(BeaconUnitTest::TickOnce(*Prerequisites));
+		UTEST_TRUE_EXPR(BeaconClient->GetConnectionState() == EBeaconConnectionState::Pending);
+		UTEST_TRUE_EXPR(BeaconClientNetStats->ReceivedControlMessages.IsEmpty());
+		UTEST_TRUE_EXPR(BeaconHostNetStats->ReceivedControlMessages == TArray<uint8>({ NMT_BeaconNetGUIDAck }));
+		BeaconHostNetStats->ReceivedControlMessages.Empty();
+
+		// Receive OnConnected RPC on the client.
+		UTEST_TRUE_EXPR(BeaconUnitTest::TickOnce(*Prerequisites));
+		UTEST_TRUE_EXPR(BeaconClient->GetConnectionState() == EBeaconConnectionState::Open);
+		UTEST_TRUE_EXPR(BeaconClientNetStats->ReceivedControlMessages.IsEmpty());
+		UTEST_TRUE_EXPR(BeaconHostNetStats->ReceivedControlMessages.IsEmpty());
+
+		// Handshake complete, channel is open.
+
+		// Disconnect the client.
+		BeaconClient->DestroyBeacon();
+		BeaconClient = nullptr;
+
+		// Cleanup successful handshake.
+		HostUserBeaconClient = BeaconUnitTest::GetBeaconClientForUser(BeaconHostObject, UserId);
+		UTEST_TRUE_EXPR(HostUserBeaconClient != nullptr);
+		UTEST_TRUE_EXPR(BeaconUnitTest::TickUntil(*Prerequisites, [&Prerequisites](){ return Prerequisites->GetStats().HostObject.NotifyClientDisconnected.InvokeCount == 1; }));
+	}
+
+	// Successful handshake.
+	// Force the use of OnReceivedNetworkEncryptionAck.
+	{
+		BeaconHostNetStats->ReceivedControlMessages.Empty();
+		Prerequisites->GetConfig() = BaseConfig;
+		Prerequisites->GetStats() = BeaconUnitTest::FTestStats();
+
+		// Setup encryption.
+
+		Prerequisites->GetConfig().Encryption.FailureAction = EEncryptionFailureAction::RejectConnection;
+
+		// Host config.
+		Prerequisites->GetConfig().Encryption.Host.bDelayDelegate = false;
+		Prerequisites->GetConfig().Encryption.Host.Response = EEncryptionResponse::Success;
+		Prerequisites->GetConfig().Encryption.Host.ErrorMsg.Empty();
+		Prerequisites->GetConfig().Encryption.Host.EncryptionData.Key = DecodedEncryptionKey;
+		Prerequisites->GetConfig().Encryption.Host.EncryptionData.Identifier = EncryptionIdentitfier;
+
+		// Client config.
+		Prerequisites->GetConfig().Encryption.Client.bDelayDelegate = false;
+		Prerequisites->GetConfig().Encryption.Client.Response = EEncryptionResponse::Success;
+		Prerequisites->GetConfig().Encryption.Client.ErrorMsg.Empty();
+		Prerequisites->GetConfig().Encryption.Client.EncryptionData.Key = DecodedEncryptionKey;
+		Prerequisites->GetConfig().Encryption.Client.EncryptionData.Identifier = EncryptionIdentitfier;
+
+		TObjectPtr<AOnlineBeaconClient> HostUserBeaconClient;
+
+		// Client setup.
+
+		TObjectPtr<AOnlineBeaconUnitTestClient> BeaconClient = Prerequisites->GetWorld()->SpawnActor<AOnlineBeaconUnitTestClient>(AOnlineBeaconUnitTestClient::StaticClass());
+		ON_SCOPE_EXIT{ if (BeaconClient) { BeaconClient->DestroyBeacon(); } };
+		UTEST_TRUE_EXPR(BeaconUnitTest::ConfigureBeacon(*Prerequisites, BeaconClient));
+
+		BeaconClient->SetEncryptionData(Prerequisites->GetConfig().Encryption.Client.EncryptionData);
+		UTEST_TRUE_EXPR(BeaconUnitTest::InitClientForUser(*Prerequisites, BeaconClient, UserId));
+		UTEST_TRUE_EXPR(Prerequisites->GetStats().Encryption.NetworkEncryptionFailure.InvokeCount == 0);
+
+		TSharedPtr<BeaconUnitTest::FNetworkStats> BeaconClientNetStats;
+		UTEST_TRUE_EXPR(BeaconUnitTest::ConfigureBeaconNetDriver(*Prerequisites, BeaconClient, &BeaconClientNetStats));
+		UTEST_TRUE_EXPR(BeaconClientNetStats.IsValid());
+
+		// Wait until connection is initialized
+		UTEST_TRUE_EXPR(BeaconUnitTest::TickUntilConnectionInitialized(*Prerequisites, BeaconClient));
+		UTEST_TRUE_EXPR(Prerequisites->GetStats().Encryption.NetworkEncryptionFailure.InvokeCount == 0);
+		UTEST_TRUE_EXPR(BeaconHostNetStats->ReceivedControlMessages.IsEmpty());
+
+		// Clearing the stored encryption info forces the use of OnReceivedNetworkEncryptionAck on the client.
+		BeaconClient->SetEncryptionData(FEncryptionData());
+
+		// Handshake testing.
+
+		UTEST_TRUE_EXPR(BeaconUnitTest::TickUntilControlMessageReceived(*Prerequisites, *BeaconHostNetStats));
+		UTEST_TRUE_EXPR(BeaconClient->GetConnectionState() == EBeaconConnectionState::Pending);
+		UTEST_TRUE_EXPR(BeaconClientNetStats->ReceivedControlMessages.IsEmpty());
+		UTEST_TRUE_EXPR(BeaconHostNetStats->ReceivedControlMessages == TArray<uint8>({ NMT_Hello }));
+		UTEST_TRUE_EXPR(Prerequisites->GetStats().Encryption.NetworkEncryptionToken.InvokeCount == 1);
+		UTEST_TRUE_EXPR(Prerequisites->GetStats().Encryption.NetworkEncryptionToken.CallbackCount == 1);
+		UTEST_TRUE_EXPR(Prerequisites->GetStats().Encryption.NetworkEncryptionFailure.InvokeCount == 0);
+		BeaconHostNetStats->ReceivedControlMessages.Empty();
+
+		UTEST_TRUE_EXPR(BeaconUnitTest::TickOnce(*Prerequisites));
+		UTEST_TRUE_EXPR(BeaconClient->GetConnectionState() == EBeaconConnectionState::Pending);
+		UTEST_TRUE_EXPR(BeaconClientNetStats->ReceivedControlMessages == TArray<uint8>({ NMT_EncryptionAck, NMT_BeaconWelcome }));
+		UTEST_TRUE_EXPR(BeaconHostNetStats->ReceivedControlMessages.IsEmpty());
+		UTEST_TRUE_EXPR(Prerequisites->GetStats().Encryption.NetworkEncryptionAck.InvokeCount == 1);
+		UTEST_TRUE_EXPR(Prerequisites->GetStats().Encryption.NetworkEncryptionAck.CallbackCount == 1);
+		UTEST_TRUE_EXPR(Prerequisites->GetStats().Encryption.NetworkEncryptionFailure.InvokeCount == 0);
+		BeaconClientNetStats->ReceivedControlMessages.Empty();
+
+		UTEST_TRUE_EXPR(BeaconUnitTest::TickOnce(*Prerequisites));
+		UTEST_TRUE_EXPR(BeaconClient->GetConnectionState() == EBeaconConnectionState::Pending);
+		UTEST_TRUE_EXPR(BeaconClientNetStats->ReceivedControlMessages.IsEmpty());
+		UTEST_TRUE_EXPR(BeaconHostNetStats->ReceivedControlMessages == TArray<uint8>({ NMT_Netspeed, NMT_BeaconJoin }));
+		BeaconHostNetStats->ReceivedControlMessages.Empty();
+
+		UTEST_TRUE_EXPR(BeaconUnitTest::TickOnce(*Prerequisites));
+		UTEST_TRUE_EXPR(BeaconClient->GetConnectionState() == EBeaconConnectionState::Pending);
+		UTEST_TRUE_EXPR(BeaconClientNetStats->ReceivedControlMessages == TArray<uint8>({ NMT_BeaconAssignGUID }));
+		UTEST_TRUE_EXPR(BeaconHostNetStats->ReceivedControlMessages.IsEmpty());
+		BeaconClientNetStats->ReceivedControlMessages.Empty();
+
+		UTEST_TRUE_EXPR(BeaconUnitTest::TickOnce(*Prerequisites));
+		UTEST_TRUE_EXPR(BeaconClient->GetConnectionState() == EBeaconConnectionState::Pending);
+		UTEST_TRUE_EXPR(BeaconClientNetStats->ReceivedControlMessages.IsEmpty());
+		UTEST_TRUE_EXPR(BeaconHostNetStats->ReceivedControlMessages == TArray<uint8>({ NMT_BeaconNetGUIDAck }));
+		BeaconHostNetStats->ReceivedControlMessages.Empty();
+
+		// Receive OnConnected RPC on the client.
+		UTEST_TRUE_EXPR(BeaconUnitTest::TickOnce(*Prerequisites));
+		UTEST_TRUE_EXPR(BeaconClient->GetConnectionState() == EBeaconConnectionState::Open);
+		UTEST_TRUE_EXPR(BeaconClientNetStats->ReceivedControlMessages.IsEmpty());
+		UTEST_TRUE_EXPR(BeaconHostNetStats->ReceivedControlMessages.IsEmpty());
+
+		// Handshake complete, channel is open.
+
+		// Disconnect the client.
+		BeaconClient->DestroyBeacon();
+		BeaconClient = nullptr;
+
+		// Cleanup successful handshake.
+		HostUserBeaconClient = BeaconUnitTest::GetBeaconClientForUser(BeaconHostObject, UserId);
+		UTEST_TRUE_EXPR(HostUserBeaconClient != nullptr);
+		UTEST_TRUE_EXPR(BeaconUnitTest::TickUntil(*Prerequisites, [&Prerequisites](){ return Prerequisites->GetStats().HostObject.NotifyClientDisconnected.InvokeCount == 1; }));
+	}
+
+	// Successful handshake.
+	// Force the use of OnReceivedNetworkEncryptionAck.
+	// Make both OnReceivedNetworkEncryptionToken and OnReceivedNetworkEncryptionAck delay callback by one frame.
+	{
+		BeaconHostNetStats->ReceivedControlMessages.Empty();
+		Prerequisites->GetConfig() = BaseConfig;
+		Prerequisites->GetStats() = BeaconUnitTest::FTestStats();
+
+		// Setup encryption.
+
+		Prerequisites->GetConfig().Encryption.FailureAction = EEncryptionFailureAction::RejectConnection;
+
+		// Host config.
+		Prerequisites->GetConfig().Encryption.Host.bDelayDelegate = true;
+		Prerequisites->GetConfig().Encryption.Host.Response = EEncryptionResponse::Success;
+		Prerequisites->GetConfig().Encryption.Host.ErrorMsg.Empty();
+		Prerequisites->GetConfig().Encryption.Host.EncryptionData.Key = DecodedEncryptionKey;
+		Prerequisites->GetConfig().Encryption.Host.EncryptionData.Identifier = EncryptionIdentitfier;
+
+		// Client config.
+		Prerequisites->GetConfig().Encryption.Client.bDelayDelegate = true;
+		Prerequisites->GetConfig().Encryption.Client.Response = EEncryptionResponse::Success;
+		Prerequisites->GetConfig().Encryption.Client.ErrorMsg.Empty();
+		Prerequisites->GetConfig().Encryption.Client.EncryptionData.Key = DecodedEncryptionKey;
+		Prerequisites->GetConfig().Encryption.Client.EncryptionData.Identifier = EncryptionIdentitfier;
+
+		TObjectPtr<AOnlineBeaconClient> HostUserBeaconClient;
+
+		// Client setup.
+
+		TObjectPtr<AOnlineBeaconUnitTestClient> BeaconClient = Prerequisites->GetWorld()->SpawnActor<AOnlineBeaconUnitTestClient>(AOnlineBeaconUnitTestClient::StaticClass());
+		ON_SCOPE_EXIT{ if (BeaconClient) { BeaconClient->DestroyBeacon(); } };
+		UTEST_TRUE_EXPR(BeaconUnitTest::ConfigureBeacon(*Prerequisites, BeaconClient));
+
+		BeaconClient->SetEncryptionData(Prerequisites->GetConfig().Encryption.Client.EncryptionData);
+		UTEST_TRUE_EXPR(BeaconUnitTest::InitClientForUser(*Prerequisites, BeaconClient, UserId));
+		UTEST_TRUE_EXPR(BeaconClient->GetConnectionState() == EBeaconConnectionState::Pending);
+		UTEST_TRUE_EXPR(BeaconHostNetStats->ReceivedControlMessages.IsEmpty());
+
+		TSharedPtr<BeaconUnitTest::FNetworkStats> BeaconClientNetStats;
+		UTEST_TRUE_EXPR(BeaconUnitTest::ConfigureBeaconNetDriver(*Prerequisites, BeaconClient, &BeaconClientNetStats));
+		UTEST_TRUE_EXPR(BeaconClientNetStats.IsValid());
+
+		// Wait until connection is initialized
+		UTEST_TRUE_EXPR(BeaconUnitTest::TickUntilConnectionInitialized(*Prerequisites, BeaconClient));
+		UTEST_TRUE_EXPR(BeaconClient->GetConnectionState() == EBeaconConnectionState::Pending);
+		UTEST_TRUE_EXPR(BeaconClientNetStats->ReceivedControlMessages.IsEmpty());
+		UTEST_TRUE_EXPR(BeaconHostNetStats->ReceivedControlMessages.IsEmpty());
+		UTEST_TRUE_EXPR(Prerequisites->GetStats().Encryption.NetworkEncryptionFailure.InvokeCount == 0);
+
+		// Clearing the stored encryption info forces the use of OnReceivedNetworkEncryptionAck on the client.
+		BeaconClient->SetEncryptionData(FEncryptionData());
+
+		// Handshake testing.
+
+		UTEST_TRUE_EXPR(BeaconUnitTest::TickUntilControlMessageReceived(*Prerequisites, *BeaconHostNetStats));
+		UTEST_TRUE_EXPR(BeaconClient->GetConnectionState() == EBeaconConnectionState::Pending);
+		UTEST_TRUE_EXPR(BeaconClientNetStats->ReceivedControlMessages.IsEmpty());
+		UTEST_TRUE_EXPR(BeaconHostNetStats->ReceivedControlMessages == TArray<uint8>({ NMT_Hello }));
+		UTEST_TRUE_EXPR(Prerequisites->GetStats().Encryption.NetworkEncryptionToken.InvokeCount == 1);
+		UTEST_TRUE_EXPR(Prerequisites->GetStats().Encryption.NetworkEncryptionToken.CallbackCount == 0);
+		UTEST_TRUE_EXPR(Prerequisites->GetStats().Encryption.NetworkEncryptionFailure.InvokeCount == 0);
+		BeaconHostNetStats->ReceivedControlMessages.Empty();
+
+		// Tick to allow the host to handle OnReceivedNetworkEncryptionToken.
+		UTEST_TRUE_EXPR(BeaconUnitTest::TickOnce(*Prerequisites));
+		UTEST_TRUE_EXPR(BeaconClient->GetConnectionState() == EBeaconConnectionState::Pending);
+		UTEST_TRUE_EXPR(BeaconClientNetStats->ReceivedControlMessages.IsEmpty());
+		UTEST_TRUE_EXPR(BeaconHostNetStats->ReceivedControlMessages.IsEmpty());
+		UTEST_TRUE_EXPR(Prerequisites->GetStats().Encryption.NetworkEncryptionToken.CallbackCount == 1);
+
+		UTEST_TRUE_EXPR(BeaconUnitTest::TickOnce(*Prerequisites));
+		UTEST_TRUE_EXPR(BeaconClient->GetConnectionState() == EBeaconConnectionState::Pending);
+		UTEST_TRUE_EXPR(BeaconClientNetStats->ReceivedControlMessages == TArray<uint8>({ NMT_EncryptionAck, NMT_BeaconWelcome }));
+		UTEST_TRUE_EXPR(BeaconHostNetStats->ReceivedControlMessages.IsEmpty());
+		UTEST_TRUE_EXPR(Prerequisites->GetStats().Encryption.NetworkEncryptionAck.InvokeCount == 1);
+		UTEST_TRUE_EXPR(Prerequisites->GetStats().Encryption.NetworkEncryptionAck.CallbackCount == 0);
+		UTEST_TRUE_EXPR(Prerequisites->GetStats().Encryption.NetworkEncryptionFailure.InvokeCount == 0);
+		BeaconClientNetStats->ReceivedControlMessages.Empty();
+
+		// Tick to allow the client to handle OnReceivedNetworkEncryptionAck.
+		UTEST_TRUE_EXPR(BeaconUnitTest::TickOnce(*Prerequisites));
+		UTEST_TRUE_EXPR(BeaconClient->GetConnectionState() == EBeaconConnectionState::Pending);
+		UTEST_TRUE_EXPR(BeaconClientNetStats->ReceivedControlMessages.IsEmpty());
+		UTEST_TRUE_EXPR(BeaconHostNetStats->ReceivedControlMessages == TArray<uint8>({ NMT_Netspeed, NMT_BeaconJoin }));
+		UTEST_TRUE_EXPR(Prerequisites->GetStats().Encryption.NetworkEncryptionAck.CallbackCount == 1);
+		BeaconHostNetStats->ReceivedControlMessages.Empty();
+
+		UTEST_TRUE_EXPR(BeaconUnitTest::TickOnce(*Prerequisites));
+		UTEST_TRUE_EXPR(BeaconClient->GetConnectionState() == EBeaconConnectionState::Pending);
+		UTEST_TRUE_EXPR(BeaconClientNetStats->ReceivedControlMessages == TArray<uint8>({ NMT_BeaconAssignGUID }));
+		UTEST_TRUE_EXPR(BeaconHostNetStats->ReceivedControlMessages.IsEmpty());
+		BeaconClientNetStats->ReceivedControlMessages.Empty();
+
+		UTEST_TRUE_EXPR(BeaconUnitTest::TickOnce(*Prerequisites));
+		UTEST_TRUE_EXPR(BeaconClient->GetConnectionState() == EBeaconConnectionState::Pending);
+		UTEST_TRUE_EXPR(BeaconClientNetStats->ReceivedControlMessages.IsEmpty());
+		UTEST_TRUE_EXPR(BeaconHostNetStats->ReceivedControlMessages == TArray<uint8>({ NMT_BeaconNetGUIDAck }));
+		BeaconHostNetStats->ReceivedControlMessages.Empty();
+
+		// Receive OnConnected RPC on the client.
+		UTEST_TRUE_EXPR(BeaconUnitTest::TickOnce(*Prerequisites));
+		UTEST_TRUE_EXPR(BeaconClient->GetConnectionState() == EBeaconConnectionState::Open);
+		UTEST_TRUE_EXPR(BeaconClientNetStats->ReceivedControlMessages.IsEmpty());
+		UTEST_TRUE_EXPR(BeaconHostNetStats->ReceivedControlMessages.IsEmpty());
+
+		// Handshake complete, channel is open.
+
+		// Disconnect the client.
+		BeaconClient->DestroyBeacon();
+		BeaconClient = nullptr;
+
+		// Cleanup successful handshake.
+		HostUserBeaconClient = BeaconUnitTest::GetBeaconClientForUser(BeaconHostObject, UserId);
+		UTEST_TRUE_EXPR(HostUserBeaconClient != nullptr);
+		UTEST_TRUE_EXPR(BeaconUnitTest::TickUntil(*Prerequisites, [&Prerequisites](){ return Prerequisites->GetStats().HostObject.NotifyClientDisconnected.InvokeCount == 1; }));
+	}
+
+	// Successful handshake.
+	// 1. Netdriver requires encryption
+	// 2. No encryption data provided.
+	// 3. Client initializes connection.
+	// 3a. BeaconEncryptionFailureAction returns AllowConnection.
+	// 4. Host allows the connection.
+	// 4a. BeaconEncryptionFailureAction returns AllowConnection.
+	{
+		BeaconHostNetStats->ReceivedControlMessages.Empty();
+		Prerequisites->GetConfig() = BaseConfig;
+		Prerequisites->GetStats() = BeaconUnitTest::FTestStats();
+
+		// Setup encryption.
+		Prerequisites->GetConfig().Encryption.FailureAction = EEncryptionFailureAction::AllowConnection;
+
+		TObjectPtr<AOnlineBeaconClient> HostUserBeaconClient;
+
+		// Client setup.
+
+		TObjectPtr<AOnlineBeaconUnitTestClient> BeaconClient = Prerequisites->GetWorld()->SpawnActor<AOnlineBeaconUnitTestClient>(AOnlineBeaconUnitTestClient::StaticClass());
+		ON_SCOPE_EXIT{ if (BeaconClient) { BeaconClient->DestroyBeacon(); } };
+		UTEST_TRUE_EXPR(BeaconUnitTest::ConfigureBeacon(*Prerequisites, BeaconClient));
+
+		UTEST_TRUE_EXPR(BeaconUnitTest::InitClientForUser(*Prerequisites, BeaconClient, UserId));
+		UTEST_TRUE_EXPR(BeaconClient->GetConnectionState() == EBeaconConnectionState::Pending);
+		UTEST_TRUE_EXPR(BeaconHostNetStats->ReceivedControlMessages.IsEmpty());
+
+		TSharedPtr<BeaconUnitTest::FNetworkStats> BeaconClientNetStats;
+		UTEST_TRUE_EXPR(BeaconUnitTest::ConfigureBeaconNetDriver(*Prerequisites, BeaconClient, &BeaconClientNetStats));
+		UTEST_TRUE_EXPR(BeaconClientNetStats.IsValid());
+
+		// Wait until connection is initialized
+		UTEST_TRUE_EXPR(BeaconUnitTest::TickUntilConnectionInitialized(*Prerequisites, BeaconClient));
+		UTEST_TRUE_EXPR(BeaconClient->GetConnectionState() == EBeaconConnectionState::Pending);
+		UTEST_TRUE_EXPR(BeaconClientNetStats->ReceivedControlMessages.IsEmpty());
+		UTEST_TRUE_EXPR(BeaconHostNetStats->ReceivedControlMessages.IsEmpty());
+		UTEST_TRUE_EXPR(Prerequisites->GetStats().Encryption.NetworkEncryptionFailure.InvokeCount == 1);
+
+		// Handshake testing.
+
+		UTEST_TRUE_EXPR(BeaconUnitTest::TickUntilControlMessageReceived(*Prerequisites, *BeaconHostNetStats));
+		UTEST_TRUE_EXPR(BeaconClient->GetConnectionState() == EBeaconConnectionState::Pending);
+		UTEST_TRUE_EXPR(BeaconClientNetStats->ReceivedControlMessages.IsEmpty());
+		UTEST_TRUE_EXPR(BeaconHostNetStats->ReceivedControlMessages == TArray<uint8>({ NMT_Hello }));
+		UTEST_TRUE_EXPR(Prerequisites->GetStats().Encryption.NetworkEncryptionToken.InvokeCount == 0);
+		UTEST_TRUE_EXPR(Prerequisites->GetStats().Encryption.NetworkEncryptionFailure.InvokeCount == 2);
+		BeaconHostNetStats->ReceivedControlMessages.Empty();
+
+		UTEST_TRUE_EXPR(BeaconUnitTest::TickOnce(*Prerequisites));
+		UTEST_TRUE_EXPR(BeaconClient->GetConnectionState() == EBeaconConnectionState::Pending);
+		UTEST_TRUE_EXPR(BeaconClientNetStats->ReceivedControlMessages == TArray<uint8>({ NMT_BeaconWelcome }));
+		UTEST_TRUE_EXPR(BeaconHostNetStats->ReceivedControlMessages.IsEmpty());
+		UTEST_TRUE_EXPR(Prerequisites->GetStats().Encryption.NetworkEncryptionAck.InvokeCount == 0);
+		BeaconClientNetStats->ReceivedControlMessages.Empty();
+
+		UTEST_TRUE_EXPR(BeaconUnitTest::TickOnce(*Prerequisites));
+		UTEST_TRUE_EXPR(BeaconClient->GetConnectionState() == EBeaconConnectionState::Pending);
+		UTEST_TRUE_EXPR(BeaconClientNetStats->ReceivedControlMessages.IsEmpty());
+		UTEST_TRUE_EXPR(BeaconHostNetStats->ReceivedControlMessages == TArray<uint8>({ NMT_Netspeed, NMT_BeaconJoin }));
+		BeaconHostNetStats->ReceivedControlMessages.Empty();
+
+		UTEST_TRUE_EXPR(BeaconUnitTest::TickOnce(*Prerequisites));
+		UTEST_TRUE_EXPR(BeaconClient->GetConnectionState() == EBeaconConnectionState::Pending);
+		UTEST_TRUE_EXPR(BeaconClientNetStats->ReceivedControlMessages == TArray<uint8>({ NMT_BeaconAssignGUID }));
+		UTEST_TRUE_EXPR(BeaconHostNetStats->ReceivedControlMessages.IsEmpty());
+		BeaconClientNetStats->ReceivedControlMessages.Empty();
+
+		UTEST_TRUE_EXPR(BeaconUnitTest::TickOnce(*Prerequisites));
+		UTEST_TRUE_EXPR(BeaconClient->GetConnectionState() == EBeaconConnectionState::Pending);
+		UTEST_TRUE_EXPR(BeaconClientNetStats->ReceivedControlMessages.IsEmpty());
+		UTEST_TRUE_EXPR(BeaconHostNetStats->ReceivedControlMessages == TArray<uint8>({ NMT_BeaconNetGUIDAck }));
+		BeaconHostNetStats->ReceivedControlMessages.Empty();
+
+		// Receive OnConnected RPC on the client.
+		UTEST_TRUE_EXPR(BeaconUnitTest::TickOnce(*Prerequisites));
+		UTEST_TRUE_EXPR(BeaconClient->GetConnectionState() == EBeaconConnectionState::Open);
+		UTEST_TRUE_EXPR(BeaconClientNetStats->ReceivedControlMessages.IsEmpty());
+		UTEST_TRUE_EXPR(BeaconHostNetStats->ReceivedControlMessages.IsEmpty());
+
+		// Handshake complete, channel is open.
+
+		// Disconnect the client.
+		BeaconClient->DestroyBeacon();
+		BeaconClient = nullptr;
+
+		// Cleanup successful handshake.
+		HostUserBeaconClient = BeaconUnitTest::GetBeaconClientForUser(BeaconHostObject, UserId);
+		UTEST_TRUE_EXPR(HostUserBeaconClient != nullptr);
+		UTEST_TRUE_EXPR(BeaconUnitTest::TickUntil(*Prerequisites, [&Prerequisites](){ return Prerequisites->GetStats().HostObject.NotifyClientDisconnected.InvokeCount == 1; }));
+	}
+
+	return true;
+}
+
+IMPLEMENT_SIMPLE_AUTOMATION_TEST(FOnlineBeaconTestEncryptedHandshakeFailed,
+	"System.Engine.Online.OnlineSubsystemUtils.OnlineBeacon.EncryptedHandshakeFailed",
+	EAutomationTestFlags::EditorContext | EAutomationTestFlags::EngineFilter)
+bool FOnlineBeaconTestEncryptedHandshakeFailed::RunTest(const FString& Parameters)
+{
+	TSharedPtr<BeaconUnitTest::FTestPrerequisites> Prerequisites = BeaconUnitTest::FTestPrerequisites::TryCreate();
+	UTEST_TRUE_EXPR(Prerequisites.IsValid());
+
+	// Check that encryption is allowed.
+	UTEST_TRUE_EXPR(CVarNetAllowEncryption.GetValueOnGameThread() != 0);
+
+	Prerequisites->GetConfig().NetDriver.bNoTimeouts = true;
+	Prerequisites->GetConfig().NetDriver.ServerListenPort = 9999;
+
+	// Todo: Remove external dependency on AESGCMHandlerComponent.
+	Prerequisites->GetConfig().Encryption.bEnabled = true;
+	Prerequisites->GetConfig().Encryption.NetDriverEncryptionComponentName = TEXT("AESGCMHandlerComponent");
+
+	// Snapshot config to restore before running a test section.
+	const BeaconUnitTest::FTestConfig BaseConfig = Prerequisites->GetConfig();
+
+	// Valid encryption key.
+	const FString EncryptionIdentitfier = TEXT("test");
+	const FString Base64EncryptionKey = TEXT("IYaVIE38d6J9VfbPULuSMfn3/axig797U8DVJyRm1/c=");
+	TArray<uint8> DecodedEncryptionKey;
+	UTEST_TRUE_EXPR(FBase64::Decode(Base64EncryptionKey, DecodedEncryptionKey));
+
+	FUniqueNetIdStringRef UserId = FUniqueNetIdString::Create(TEXT("User"), TEXT("UnitTest"));
+
+	// Host setup.
+	TObjectPtr<AOnlineBeaconUnitTestHostObject> BeaconHostObject = Prerequisites->GetWorld()->SpawnActor<AOnlineBeaconUnitTestHostObject>(AOnlineBeaconUnitTestHostObject::StaticClass());
+	TObjectPtr<AOnlineBeaconUnitTestHost> BeaconHost = Prerequisites->GetWorld()->SpawnActor<AOnlineBeaconUnitTestHost>(AOnlineBeaconUnitTestHost::StaticClass());
+	ON_SCOPE_EXIT{ if (BeaconHostObject) { BeaconHostObject->Destroy(); } };
+	ON_SCOPE_EXIT{ if (BeaconHost) { BeaconHost->DestroyBeacon(); } };
+	UTEST_TRUE_EXPR(BeaconUnitTest::ConfigureBeacon(*Prerequisites, BeaconHost));
+	BeaconHost->ListenPort = Prerequisites->GetConfig().NetDriver.ServerListenPort;
+	BeaconHost->RegisterHost(BeaconHostObject);
+	UTEST_TRUE_EXPR(BeaconHost->InitHost());
+	BeaconHost->PauseBeaconRequests(false);
+
+	TSharedPtr<BeaconUnitTest::FNetworkStats> BeaconHostNetStats;
+	UTEST_TRUE_EXPR(BeaconUnitTest::ConfigureBeaconNetDriver(*Prerequisites, BeaconHost, &BeaconHostNetStats));
+	UTEST_TRUE_EXPR(BeaconHostNetStats.IsValid());
+
+	// Failed handshake. OnReceivedNetworkEncryptionToken.
+	// 1. Netdriver requires encryption
+	// 2. Encryption data provided.
+	// 3. Client initializes connection.
+	// 4. Host fails the connection.
+	// 4a. FNetDelegates::OnReceivedNetworkEncryptionToken returns EEncryptionResponse::Failure
+	{
+		BeaconHostNetStats->ReceivedControlMessages.Empty();
+		Prerequisites->GetConfig() = BaseConfig;
+		Prerequisites->GetStats() = BeaconUnitTest::FTestStats();
+
+		Prerequisites->GetConfig().Encryption.FailureAction = EEncryptionFailureAction::RejectConnection;
+
+		// Host config.
+		Prerequisites->GetConfig().Encryption.Host.bDelayDelegate = false;
+		Prerequisites->GetConfig().Encryption.Host.Response = EEncryptionResponse::Failure;
+		Prerequisites->GetConfig().Encryption.Host.ErrorMsg.Empty();
+		Prerequisites->GetConfig().Encryption.Host.EncryptionData.Key = DecodedEncryptionKey;
+		Prerequisites->GetConfig().Encryption.Host.EncryptionData.Identifier = EncryptionIdentitfier;
+
+		// Client config.
+		Prerequisites->GetConfig().Encryption.Client.bDelayDelegate = false;
+		Prerequisites->GetConfig().Encryption.Client.Response = EEncryptionResponse::Success;
+		Prerequisites->GetConfig().Encryption.Client.ErrorMsg.Empty();
+		Prerequisites->GetConfig().Encryption.Client.EncryptionData.Key = DecodedEncryptionKey;
+		Prerequisites->GetConfig().Encryption.Client.EncryptionData.Identifier = EncryptionIdentitfier;
+
+		// Client setup.
+
+		TObjectPtr<AOnlineBeaconUnitTestClient> BeaconClient = Prerequisites->GetWorld()->SpawnActor<AOnlineBeaconUnitTestClient>(AOnlineBeaconUnitTestClient::StaticClass());
+		ON_SCOPE_EXIT{ if (BeaconClient) { BeaconClient->DestroyBeacon(); } };
+		UTEST_TRUE_EXPR(BeaconUnitTest::ConfigureBeacon(*Prerequisites, BeaconClient));
+
+		BeaconClient->SetEncryptionData(Prerequisites->GetConfig().Encryption.Client.EncryptionData);
+		UTEST_TRUE_EXPR(BeaconUnitTest::InitClientForUser(*Prerequisites, BeaconClient, UserId));
+		UTEST_TRUE_EXPR(BeaconClient->GetConnectionState() == EBeaconConnectionState::Pending);
+		UTEST_TRUE_EXPR(Prerequisites->GetStats().Encryption.NetworkEncryptionFailure.InvokeCount == 0);
+
+		TSharedPtr<BeaconUnitTest::FNetworkStats> BeaconClientNetStats;
+		UTEST_TRUE_EXPR(BeaconUnitTest::ConfigureBeaconNetDriver(*Prerequisites, BeaconClient, &BeaconClientNetStats));
+		UTEST_TRUE_EXPR(BeaconClientNetStats.IsValid());
+
+		// Wait until connection is initialized
+		UTEST_TRUE_EXPR(BeaconUnitTest::TickUntilConnectionInitialized(*Prerequisites, BeaconClient));
+		UTEST_TRUE_EXPR(BeaconClient->GetConnectionState() == EBeaconConnectionState::Pending);
+		UTEST_TRUE_EXPR(BeaconClientNetStats->ReceivedControlMessages.IsEmpty());
+		UTEST_TRUE_EXPR(BeaconHostNetStats->ReceivedControlMessages.IsEmpty());
+
+		// Handshake testing.
+
+		UTEST_TRUE_EXPR(BeaconUnitTest::TickUntilControlMessageReceived(*Prerequisites, *BeaconHostNetStats));
+		UTEST_TRUE_EXPR(BeaconClient->GetConnectionState() == EBeaconConnectionState::Pending);
+		UTEST_TRUE_EXPR(BeaconClientNetStats->ReceivedControlMessages.IsEmpty());
+		UTEST_TRUE_EXPR(BeaconHostNetStats->ReceivedControlMessages == TArray<uint8>({ NMT_Hello }));
+		BeaconHostNetStats->ReceivedControlMessages.Empty();
+
+		UTEST_TRUE_EXPR(BeaconUnitTest::TickUntilDisconnected(*Prerequisites, BeaconClient));
+		UTEST_TRUE_EXPR(BeaconClient->GetConnectionState() == EBeaconConnectionState::Invalid);
+		UTEST_TRUE_EXPR(BeaconClientNetStats->ReceivedControlMessages == TArray<uint8>({ NMT_Failure }));
+		UTEST_TRUE_EXPR(Prerequisites->GetStats().Client.OnFailure.InvokeCount == 1);
+		UTEST_TRUE_EXPR(Prerequisites->GetStats().Encryption.NetworkEncryptionToken.InvokeCount == 1);
+		UTEST_TRUE_EXPR(Prerequisites->GetStats().Encryption.NetworkEncryptionAck.InvokeCount == 0);
+		UTEST_TRUE_EXPR(Prerequisites->GetStats().Encryption.NetworkEncryptionFailure.InvokeCount == 0);
+	}
+
+	// Failed handshake. OnReceivedNetworkEncryptionAck.
+	// 1. Netdriver requires encryption
+	// 2. Encryption data provided.
+	// 3. Client initializes connection.
+	// 4. Host allows the connection.
+	// 5. Client fails the connection.
+	// 5a. FNetDelegates::OnReceivedNetworkEncryptionAck returns EEncryptionResponse::Failure
+	{
+		BeaconHostNetStats->ReceivedControlMessages.Empty();
+		Prerequisites->GetConfig() = BaseConfig;
+		Prerequisites->GetStats() = BeaconUnitTest::FTestStats();
+
+		Prerequisites->GetConfig().Encryption.FailureAction = EEncryptionFailureAction::RejectConnection;
+
+		// Host config.
+		Prerequisites->GetConfig().Encryption.Host.bDelayDelegate = false;
+		Prerequisites->GetConfig().Encryption.Host.Response = EEncryptionResponse::Success;
+		Prerequisites->GetConfig().Encryption.Host.ErrorMsg.Empty();
+		Prerequisites->GetConfig().Encryption.Host.EncryptionData.Key = DecodedEncryptionKey;
+		Prerequisites->GetConfig().Encryption.Host.EncryptionData.Identifier = EncryptionIdentitfier;
+
+		// Client config.
+		Prerequisites->GetConfig().Encryption.Client.bDelayDelegate = false;
+		Prerequisites->GetConfig().Encryption.Client.Response = EEncryptionResponse::Failure;
+		Prerequisites->GetConfig().Encryption.Client.ErrorMsg.Empty();
+		Prerequisites->GetConfig().Encryption.Client.EncryptionData.Key = DecodedEncryptionKey;
+		Prerequisites->GetConfig().Encryption.Client.EncryptionData.Identifier = EncryptionIdentitfier;
+
+		// Client setup.
+
+		TObjectPtr<AOnlineBeaconUnitTestClient> BeaconClient = Prerequisites->GetWorld()->SpawnActor<AOnlineBeaconUnitTestClient>(AOnlineBeaconUnitTestClient::StaticClass());
+		ON_SCOPE_EXIT{ if (BeaconClient) { BeaconClient->DestroyBeacon(); } };
+		UTEST_TRUE_EXPR(BeaconUnitTest::ConfigureBeacon(*Prerequisites, BeaconClient));
+
+		BeaconClient->SetEncryptionData(Prerequisites->GetConfig().Encryption.Client.EncryptionData);
+		UTEST_TRUE_EXPR(BeaconUnitTest::InitClientForUser(*Prerequisites, BeaconClient, UserId));
+		UTEST_TRUE_EXPR(Prerequisites->GetStats().Encryption.NetworkEncryptionFailure.InvokeCount == 0);
+
+		TSharedPtr<BeaconUnitTest::FNetworkStats> BeaconClientNetStats;
+		UTEST_TRUE_EXPR(BeaconUnitTest::ConfigureBeaconNetDriver(*Prerequisites, BeaconClient, &BeaconClientNetStats));
+		UTEST_TRUE_EXPR(BeaconClientNetStats.IsValid());
+
+		// Wait until connection is initialized
+		UTEST_TRUE_EXPR(BeaconUnitTest::TickUntilConnectionInitialized(*Prerequisites, BeaconClient));
+		UTEST_TRUE_EXPR(BeaconClient->GetConnectionState() == EBeaconConnectionState::Pending);
+		UTEST_TRUE_EXPR(BeaconClientNetStats->ReceivedControlMessages.IsEmpty());
+		UTEST_TRUE_EXPR(BeaconHostNetStats->ReceivedControlMessages.IsEmpty());
+		UTEST_TRUE_EXPR(Prerequisites->GetStats().Encryption.NetworkEncryptionFailure.InvokeCount == 0);
+
+		// Clearing the stored encryption info forces the use of OnReceivedNetworkEncryptionAck on the client.
+		BeaconClient->SetEncryptionData(FEncryptionData());
+
+		// Handshake testing.
+
+		UTEST_TRUE_EXPR(BeaconUnitTest::TickUntilControlMessageReceived(*Prerequisites, *BeaconHostNetStats));
+		UTEST_TRUE_EXPR(BeaconClient->GetConnectionState() == EBeaconConnectionState::Pending);
+		UTEST_TRUE_EXPR(BeaconClientNetStats->ReceivedControlMessages.IsEmpty());
+		UTEST_TRUE_EXPR(BeaconHostNetStats->ReceivedControlMessages == TArray<uint8>({ NMT_Hello }));
+		UTEST_TRUE_EXPR(Prerequisites->GetStats().Encryption.NetworkEncryptionToken.InvokeCount == 1);
+		UTEST_TRUE_EXPR(Prerequisites->GetStats().Encryption.NetworkEncryptionToken.CallbackCount == 1);
+		UTEST_TRUE_EXPR(Prerequisites->GetStats().Encryption.NetworkEncryptionFailure.InvokeCount == 0);
+		BeaconHostNetStats->ReceivedControlMessages.Empty();
+		Prerequisites->GetStats().Encryption = BeaconUnitTest::FTestStats::FEncryption();
+
+		UTEST_TRUE_EXPR(BeaconUnitTest::TickOnce(*Prerequisites));
+		UTEST_TRUE_EXPR(BeaconClient->GetConnectionState() == EBeaconConnectionState::Invalid);
+		UTEST_TRUE_EXPR(BeaconClientNetStats->ReceivedControlMessages == TArray<uint8>({ NMT_EncryptionAck }));
+		UTEST_TRUE_EXPR(BeaconHostNetStats->ReceivedControlMessages.IsEmpty());
+		UTEST_TRUE_EXPR(Prerequisites->GetStats().Encryption.NetworkEncryptionAck.InvokeCount == 1);
+		UTEST_TRUE_EXPR(Prerequisites->GetStats().Encryption.NetworkEncryptionAck.CallbackCount == 1);
+		UTEST_TRUE_EXPR(Prerequisites->GetStats().Encryption.NetworkEncryptionFailure.InvokeCount == 0);
+	}
+
+	// Failed handshake. BeaconEncryptionFailureAction.
+	// 1. Netdriver requires encryption
+	// 2. No encryption data provided.
+	// 3. Client initializes connection.
+	// 3a. BeaconEncryptionFailureAction returns AllowConnection.
+	// 4. Host disallows the connection.
+	// 4a. BeaconEncryptionFailureAction returns Default.
+	{
+		BeaconHostNetStats->ReceivedControlMessages.Empty();
+		Prerequisites->GetConfig() = BaseConfig;
+		Prerequisites->GetStats() = BeaconUnitTest::FTestStats();
+
+		// Client setup.
+
+		TObjectPtr<AOnlineBeaconUnitTestClient> BeaconClient = Prerequisites->GetWorld()->SpawnActor<AOnlineBeaconUnitTestClient>(AOnlineBeaconUnitTestClient::StaticClass());
+		ON_SCOPE_EXIT{ if (BeaconClient) { BeaconClient->DestroyBeacon(); } };
+		UTEST_TRUE_EXPR(BeaconUnitTest::ConfigureBeacon(*Prerequisites, BeaconClient));
+
+		Prerequisites->GetConfig().Encryption.FailureAction = EEncryptionFailureAction::AllowConnection;
+
+		UTEST_TRUE_EXPR(BeaconUnitTest::InitClientForUser(*Prerequisites, BeaconClient, UserId));
+		UTEST_TRUE_EXPR(BeaconClient->GetConnectionState() == EBeaconConnectionState::Pending);
+		
+		TSharedPtr<BeaconUnitTest::FNetworkStats> BeaconClientNetStats;
+		UTEST_TRUE_EXPR(BeaconUnitTest::ConfigureBeaconNetDriver(*Prerequisites, BeaconClient, &BeaconClientNetStats));
+		UTEST_TRUE_EXPR(BeaconClientNetStats.IsValid());
+
+		// Wait until connection is initialized
+		UTEST_TRUE_EXPR(BeaconUnitTest::TickUntilConnectionInitialized(*Prerequisites, BeaconClient));
+		UTEST_TRUE_EXPR(BeaconClient->GetConnectionState() == EBeaconConnectionState::Pending);
+		UTEST_TRUE_EXPR(BeaconClientNetStats->ReceivedControlMessages.IsEmpty());
+		UTEST_TRUE_EXPR(BeaconHostNetStats->ReceivedControlMessages.IsEmpty());
+		UTEST_TRUE_EXPR(Prerequisites->GetStats().Encryption.NetworkEncryptionFailure.InvokeCount == 1);
+
+		// Handshake testing.
+
+		Prerequisites->GetConfig().Encryption.FailureAction = EEncryptionFailureAction::Default;
+
+		UTEST_TRUE_EXPR(BeaconUnitTest::TickUntilControlMessageReceived(*Prerequisites, *BeaconHostNetStats));
+		UTEST_TRUE_EXPR(BeaconClient->GetConnectionState() == EBeaconConnectionState::Pending);
+		UTEST_TRUE_EXPR(BeaconClientNetStats->ReceivedControlMessages.IsEmpty());
+		UTEST_TRUE_EXPR(BeaconHostNetStats->ReceivedControlMessages == TArray<uint8>({ NMT_Hello }));
+		UTEST_TRUE_EXPR(Prerequisites->GetStats().Encryption.NetworkEncryptionToken.InvokeCount == 0);
+		UTEST_TRUE_EXPR(Prerequisites->GetStats().Encryption.NetworkEncryptionFailure.InvokeCount == 2);
+		BeaconHostNetStats->ReceivedControlMessages.Empty();
+
+		UTEST_TRUE_EXPR(BeaconUnitTest::TickUntilDisconnected(*Prerequisites, BeaconClient));
+		UTEST_TRUE_EXPR(BeaconClient->GetConnectionState() == EBeaconConnectionState::Invalid);
+		UTEST_TRUE_EXPR(BeaconClientNetStats->ReceivedControlMessages == TArray<uint8>({ NMT_Failure }));
+		UTEST_TRUE_EXPR(BeaconHostNetStats->ReceivedControlMessages.IsEmpty());
+		UTEST_TRUE_EXPR(Prerequisites->GetStats().Client.OnFailure.InvokeCount == 1);
+	}
+
+	// Failed handshake. BeaconEncryptionFailureAction.
+	// 1. Netdriver requires encryption
+	// 2. No encryption data provided.
+	// 3. Client initializes connection.
+	// 3a. BeaconEncryptionFailureAction returns AllowConnection.
+	// 4. Host disallows the connection.
+	// 4a. BeaconEncryptionFailureAction returns RejectConnection.
+	{
+		BeaconHostNetStats->ReceivedControlMessages.Empty();
+		Prerequisites->GetConfig() = BaseConfig;
+		Prerequisites->GetStats() = BeaconUnitTest::FTestStats();
+
+		// Client setup.
+
+		TObjectPtr<AOnlineBeaconUnitTestClient> BeaconClient = Prerequisites->GetWorld()->SpawnActor<AOnlineBeaconUnitTestClient>(AOnlineBeaconUnitTestClient::StaticClass());
+		ON_SCOPE_EXIT{ if (BeaconClient) { BeaconClient->DestroyBeacon(); } };
+		UTEST_TRUE_EXPR(BeaconUnitTest::ConfigureBeacon(*Prerequisites, BeaconClient));
+
+		Prerequisites->GetConfig().Encryption.FailureAction = EEncryptionFailureAction::AllowConnection;
+
+		UTEST_TRUE_EXPR(BeaconUnitTest::InitClientForUser(*Prerequisites, BeaconClient, UserId));
+		UTEST_TRUE_EXPR(BeaconClient->GetConnectionState() == EBeaconConnectionState::Pending);
+
+		TSharedPtr<BeaconUnitTest::FNetworkStats> BeaconClientNetStats;
+		UTEST_TRUE_EXPR(BeaconUnitTest::ConfigureBeaconNetDriver(*Prerequisites, BeaconClient, &BeaconClientNetStats));
+		UTEST_TRUE_EXPR(BeaconClientNetStats.IsValid());
+
+		// Wait until connection is initialized
+		UTEST_TRUE_EXPR(BeaconUnitTest::TickUntilConnectionInitialized(*Prerequisites, BeaconClient));
+		UTEST_TRUE_EXPR(BeaconClient->GetConnectionState() == EBeaconConnectionState::Pending);
+		UTEST_TRUE_EXPR(BeaconClientNetStats->ReceivedControlMessages.IsEmpty());
+		UTEST_TRUE_EXPR(BeaconHostNetStats->ReceivedControlMessages.IsEmpty());
+		UTEST_TRUE_EXPR(Prerequisites->GetStats().Encryption.NetworkEncryptionFailure.InvokeCount == 1);
+
+		// Handshake testing.
+
+		Prerequisites->GetConfig().Encryption.FailureAction = EEncryptionFailureAction::RejectConnection;
+
+		UTEST_TRUE_EXPR(BeaconUnitTest::TickUntilControlMessageReceived(*Prerequisites, *BeaconHostNetStats));
+		UTEST_TRUE_EXPR(BeaconClientNetStats->ReceivedControlMessages.IsEmpty());
+		UTEST_TRUE_EXPR(BeaconHostNetStats->ReceivedControlMessages == TArray<uint8>({ NMT_Hello }));
+		UTEST_TRUE_EXPR(Prerequisites->GetStats().Encryption.NetworkEncryptionToken.InvokeCount == 0);
+		UTEST_TRUE_EXPR(Prerequisites->GetStats().Encryption.NetworkEncryptionFailure.InvokeCount == 2);
+		BeaconHostNetStats->ReceivedControlMessages.Empty();
+
+		UTEST_TRUE_EXPR(BeaconUnitTest::TickUntilDisconnected(*Prerequisites, BeaconClient));
+		UTEST_TRUE_EXPR(BeaconClient->GetConnectionState() == EBeaconConnectionState::Invalid);
+		UTEST_TRUE_EXPR(BeaconClientNetStats->ReceivedControlMessages == TArray<uint8>({ NMT_Failure }));
+		UTEST_TRUE_EXPR(BeaconHostNetStats->ReceivedControlMessages.IsEmpty());
+		UTEST_TRUE_EXPR(Prerequisites->GetStats().Client.OnFailure.InvokeCount == 1);
+	}
+
+	// Failed handshake. BeaconEncryptionFailureAction.
+	// 1. Netdriver requires encryption
+	// 2. No encryption data provided.
+	// 3. Client initializes connection.
+	// 3a. BeaconEncryptionFailureAction returns Default.
+	{
+		BeaconHostNetStats->ReceivedControlMessages.Empty();
+		Prerequisites->GetConfig() = BaseConfig;
+		Prerequisites->GetStats() = BeaconUnitTest::FTestStats();
+
+		// Setup encryption.
+		Prerequisites->GetConfig().Encryption.FailureAction = EEncryptionFailureAction::Default;
+
+		// Client setup.
+
+		TObjectPtr<AOnlineBeaconUnitTestClient> BeaconClient = Prerequisites->GetWorld()->SpawnActor<AOnlineBeaconUnitTestClient>(AOnlineBeaconUnitTestClient::StaticClass());
+		ON_SCOPE_EXIT{ if (BeaconClient) { BeaconClient->DestroyBeacon(); } };
+		UTEST_TRUE_EXPR(BeaconUnitTest::ConfigureBeacon(*Prerequisites, BeaconClient));
+
+		UTEST_TRUE_EXPR(BeaconUnitTest::InitClientForUser(*Prerequisites, BeaconClient, UserId));
+		UTEST_TRUE_EXPR(BeaconClient->GetConnectionState() == EBeaconConnectionState::Pending);
+		UTEST_TRUE_EXPR(BeaconHostNetStats->ReceivedControlMessages.IsEmpty());
+
+		TSharedPtr<BeaconUnitTest::FNetworkStats> BeaconClientNetStats;
+		UTEST_TRUE_EXPR(BeaconUnitTest::ConfigureBeaconNetDriver(*Prerequisites, BeaconClient, &BeaconClientNetStats));
+		UTEST_TRUE_EXPR(BeaconClientNetStats.IsValid());
+
+		// Wait until connection is initialized
+		UTEST_TRUE_EXPR(BeaconUnitTest::TickUntilConnectionInitialized(*Prerequisites, BeaconClient));
+		UTEST_TRUE_EXPR(BeaconClient->GetConnectionState() == EBeaconConnectionState::Invalid);
+		UTEST_TRUE_EXPR(BeaconHostNetStats->ReceivedControlMessages.IsEmpty());
+		UTEST_TRUE_EXPR(BeaconClientNetStats->ReceivedControlMessages.IsEmpty());
+		UTEST_TRUE_EXPR(Prerequisites->GetStats().Encryption.NetworkEncryptionFailure.InvokeCount == 1);
+	}
+
+	// Failed handshake. BeaconEncryptionFailureAction.
+	// 1. Netdriver requires encryption
+	// 2. No encryption data provided.
+	// 3. Client initializes connection.
+	// 3a. BeaconEncryptionFailureAction returns AllowConnection.
+	// 4. BeaconEncryptionFailureAction is unbound.
+	// 5. Host disallows the connection.
+	{
+		BeaconHostNetStats->ReceivedControlMessages.Empty();
+		Prerequisites->GetConfig() = BaseConfig;
+		Prerequisites->GetStats() = BeaconUnitTest::FTestStats();
+
+		// Setup encryption.
+		Prerequisites->GetConfig().Encryption.FailureAction = EEncryptionFailureAction::AllowConnection;
+
+		// Client setup.
+
+		TObjectPtr<AOnlineBeaconUnitTestClient> BeaconClient = Prerequisites->GetWorld()->SpawnActor<AOnlineBeaconUnitTestClient>(AOnlineBeaconUnitTestClient::StaticClass());
+		ON_SCOPE_EXIT{ if (BeaconClient) { BeaconClient->DestroyBeacon(); } };
+		UTEST_TRUE_EXPR(BeaconUnitTest::ConfigureBeacon(*Prerequisites, BeaconClient));
+
+		UTEST_TRUE_EXPR(BeaconUnitTest::InitClientForUser(*Prerequisites, BeaconClient, UserId));
+		UTEST_TRUE_EXPR(BeaconClient->GetConnectionState() == EBeaconConnectionState::Pending);
+		UTEST_TRUE_EXPR(BeaconHostNetStats->ReceivedControlMessages.IsEmpty());
+
+		TSharedPtr<BeaconUnitTest::FNetworkStats> BeaconClientNetStats;
+		UTEST_TRUE_EXPR(BeaconUnitTest::ConfigureBeaconNetDriver(*Prerequisites, BeaconClient, &BeaconClientNetStats));
+		UTEST_TRUE_EXPR(BeaconClientNetStats.IsValid());
+
+		// Wait until connection is initialized
+		UTEST_TRUE_EXPR(BeaconUnitTest::TickUntilConnectionInitialized(*Prerequisites, BeaconClient));
+		UTEST_TRUE_EXPR(BeaconClient->GetConnectionState() == EBeaconConnectionState::Pending);
+		UTEST_TRUE_EXPR(BeaconClientNetStats->ReceivedControlMessages.IsEmpty());
+		UTEST_TRUE_EXPR(BeaconHostNetStats->ReceivedControlMessages.IsEmpty());
+		UTEST_TRUE_EXPR(Prerequisites->GetStats().Encryption.NetworkEncryptionFailure.InvokeCount == 1);
+
+		Prerequisites->UnbindNetEncryptionDelegates();
+		ON_SCOPE_EXIT{ Prerequisites->BindNetEncryptionDelegates(); };
+
+		// Handshake testing.
+
+		UTEST_TRUE_EXPR(BeaconUnitTest::TickUntilControlMessageReceived(*Prerequisites, *BeaconHostNetStats));
+		UTEST_TRUE_EXPR(BeaconClient->GetConnectionState() == EBeaconConnectionState::Pending);
+		UTEST_TRUE_EXPR(BeaconClientNetStats->ReceivedControlMessages.IsEmpty());
+		UTEST_TRUE_EXPR(BeaconHostNetStats->ReceivedControlMessages == TArray<uint8>({ NMT_Hello }));
+		UTEST_TRUE_EXPR(Prerequisites->GetStats().Encryption.NetworkEncryptionToken.InvokeCount == 0);
+		UTEST_TRUE_EXPR(Prerequisites->GetStats().Encryption.NetworkEncryptionFailure.InvokeCount == 1);
+		BeaconHostNetStats->ReceivedControlMessages.Empty();
+
+		UTEST_TRUE_EXPR(BeaconUnitTest::TickUntilDisconnected(*Prerequisites, BeaconClient));
+		UTEST_TRUE_EXPR(BeaconClient->GetConnectionState() == EBeaconConnectionState::Invalid);
+		UTEST_TRUE_EXPR(BeaconClientNetStats->ReceivedControlMessages == TArray<uint8>({ NMT_Failure }));
+		UTEST_TRUE_EXPR(BeaconHostNetStats->ReceivedControlMessages.IsEmpty());
+		UTEST_TRUE_EXPR(Prerequisites->GetStats().Client.OnFailure.InvokeCount == 1);
+	}
+
+	// Failed handshake. BeaconEncryptionFailureAction.
+	// 1. Netdriver requires encryption
+	// 2. No encryption data provided.
+	// 3. Client initializes connection.
+	// 3a. BeaconEncryptionFailureAction returns RejectConnection.
+	{
+		BeaconHostNetStats->ReceivedControlMessages.Empty();
+		Prerequisites->GetConfig() = BaseConfig;
+		Prerequisites->GetStats() = BeaconUnitTest::FTestStats();
+
+		// Setup encryption.
+		Prerequisites->GetConfig().Encryption.FailureAction = EEncryptionFailureAction::RejectConnection;
+
+		// Client setup.
+
+		TObjectPtr<AOnlineBeaconUnitTestClient> BeaconClient = Prerequisites->GetWorld()->SpawnActor<AOnlineBeaconUnitTestClient>(AOnlineBeaconUnitTestClient::StaticClass());
+		ON_SCOPE_EXIT{ if (BeaconClient) { BeaconClient->DestroyBeacon(); } };
+		UTEST_TRUE_EXPR(BeaconUnitTest::ConfigureBeacon(*Prerequisites, BeaconClient));
+
+		UTEST_TRUE_EXPR(BeaconUnitTest::InitClientForUser(*Prerequisites, BeaconClient, UserId));
+		UTEST_TRUE_EXPR(BeaconClient->GetConnectionState() == EBeaconConnectionState::Pending);
+		UTEST_TRUE_EXPR(BeaconHostNetStats->ReceivedControlMessages.IsEmpty());
+
+		TSharedPtr<BeaconUnitTest::FNetworkStats> BeaconClientNetStats;
+		UTEST_TRUE_EXPR(BeaconUnitTest::ConfigureBeaconNetDriver(*Prerequisites, BeaconClient, &BeaconClientNetStats));
+		UTEST_TRUE_EXPR(BeaconClientNetStats.IsValid());
+
+		// Wait until connection is initialized
+		UTEST_TRUE_EXPR(BeaconUnitTest::TickUntilConnectionInitialized(*Prerequisites, BeaconClient));
+		UTEST_TRUE_EXPR(BeaconClient->GetConnectionState() == EBeaconConnectionState::Invalid);
+		UTEST_TRUE_EXPR(Prerequisites->GetStats().Encryption.NetworkEncryptionFailure.InvokeCount == 1);
+		UTEST_TRUE_EXPR(BeaconHostNetStats->ReceivedControlMessages.IsEmpty());
+	}
+
+	// Failed handshake. BeaconEncryptionFailureAction.
+	// 1. Netdriver requires encryption
+	// 2. No encryption data provided.
+	// 3. Client initializes connection.
+	// 3a. BeaconEncryptionFailureAction returns AllowConnection.
+	// 4. Host disallows the connection.
+	{
+		BeaconHostNetStats->ReceivedControlMessages.Empty();
+		Prerequisites->GetConfig() = BaseConfig;
+		Prerequisites->GetStats() = BeaconUnitTest::FTestStats();
+
+		// Setup encryption.
+		Prerequisites->GetConfig().Encryption.FailureAction = EEncryptionFailureAction::AllowConnection;
+
+		// Client setup.
+
+		TObjectPtr<AOnlineBeaconUnitTestClient> BeaconClient = Prerequisites->GetWorld()->SpawnActor<AOnlineBeaconUnitTestClient>(AOnlineBeaconUnitTestClient::StaticClass());
+		ON_SCOPE_EXIT{ if (BeaconClient) { BeaconClient->DestroyBeacon(); } };
+		UTEST_TRUE_EXPR(BeaconUnitTest::ConfigureBeacon(*Prerequisites, BeaconClient));
+
+		UTEST_TRUE_EXPR(BeaconUnitTest::InitClientForUser(*Prerequisites, BeaconClient, UserId));
+
+		TSharedPtr<BeaconUnitTest::FNetworkStats> BeaconClientNetStats;
+		UTEST_TRUE_EXPR(BeaconUnitTest::ConfigureBeaconNetDriver(*Prerequisites, BeaconClient, &BeaconClientNetStats));
+		UTEST_TRUE_EXPR(BeaconClientNetStats.IsValid());
+
+		// Wait until connection is initialized
+		UTEST_TRUE_EXPR(BeaconUnitTest::TickUntilConnectionInitialized(*Prerequisites, BeaconClient));
+		UTEST_TRUE_EXPR(Prerequisites->GetStats().Encryption.NetworkEncryptionFailure.InvokeCount == 1);
+
+		Prerequisites->UnbindNetEncryptionDelegates();
+		ON_SCOPE_EXIT{ Prerequisites->BindNetEncryptionDelegates(); };
+
+		// Handshake testing.
+
+		UTEST_TRUE_EXPR(BeaconUnitTest::TickUntilControlMessageReceived(*Prerequisites, *BeaconHostNetStats));
+		UTEST_TRUE_EXPR(BeaconClient->GetConnectionState() == EBeaconConnectionState::Pending);
+		UTEST_TRUE_EXPR(BeaconClientNetStats->ReceivedControlMessages.IsEmpty());
+		UTEST_TRUE_EXPR(BeaconHostNetStats->ReceivedControlMessages == TArray<uint8>({ NMT_Hello }));
+		UTEST_TRUE_EXPR(Prerequisites->GetStats().Encryption.NetworkEncryptionToken.InvokeCount == 0);
+		UTEST_TRUE_EXPR(Prerequisites->GetStats().Encryption.NetworkEncryptionFailure.InvokeCount == 1);
+		BeaconHostNetStats->ReceivedControlMessages.Empty();
+
+		UTEST_TRUE_EXPR(BeaconUnitTest::TickUntilDisconnected(*Prerequisites, BeaconClient));
+		UTEST_TRUE_EXPR(BeaconClient->GetConnectionState() == EBeaconConnectionState::Invalid);
+		UTEST_TRUE_EXPR(BeaconClientNetStats->ReceivedControlMessages == TArray<uint8>({ NMT_Failure }));
+		UTEST_TRUE_EXPR(BeaconHostNetStats->ReceivedControlMessages.IsEmpty());
+		UTEST_TRUE_EXPR(Prerequisites->GetStats().Client.OnFailure.InvokeCount == 1);
+	}
+	
+	return true;
+}
+
+IMPLEMENT_SIMPLE_AUTOMATION_TEST(FOnlineBeaconTestAuthenticatedHandshakeSuccess,
+	"System.Engine.Online.OnlineSubsystemUtils.OnlineBeacon.AuthenticatedHandshakeSuccess",
+	EAutomationTestFlags::EditorContext | EAutomationTestFlags::EngineFilter)
+bool FOnlineBeaconTestAuthenticatedHandshakeSuccess::RunTest(const FString& Parameters)
+{
+	TSharedPtr<BeaconUnitTest::FTestPrerequisites> Prerequisites = BeaconUnitTest::FTestPrerequisites::TryCreate();
+	UTEST_TRUE_EXPR(Prerequisites.IsValid());
+
+	Prerequisites->GetConfig().Auth.bEnabled = true;
+	Prerequisites->GetConfig().NetDriver.bNoTimeouts = true;
+	Prerequisites->GetConfig().NetDriver.ServerListenPort = 9999;
+
+	// Snapshot config to restore before running a test section.
+	const BeaconUnitTest::FTestConfig BaseConfig = Prerequisites->GetConfig();
+
+	FUniqueNetIdStringRef UserId = FUniqueNetIdString::Create(TEXT("User"), TEXT("UnitTest"));
+
+	// Host setup.
+	TObjectPtr<AOnlineBeaconUnitTestHostObject> BeaconHostObject = Prerequisites->GetWorld()->SpawnActor<AOnlineBeaconUnitTestHostObject>(AOnlineBeaconUnitTestHostObject::StaticClass());
+	TObjectPtr<AOnlineBeaconUnitTestHost> BeaconHost = Prerequisites->GetWorld()->SpawnActor<AOnlineBeaconUnitTestHost>(AOnlineBeaconUnitTestHost::StaticClass());
+	ON_SCOPE_EXIT{ if (BeaconHostObject) { BeaconHostObject->Destroy(); } };
+	ON_SCOPE_EXIT{ if (BeaconHost) { BeaconHost->DestroyBeacon(); } };
+	UTEST_TRUE_EXPR(BeaconUnitTest::ConfigureBeacon(*Prerequisites, BeaconHost));
+	BeaconHost->ListenPort = Prerequisites->GetConfig().NetDriver.ServerListenPort;
+	BeaconHost->RegisterHost(BeaconHostObject);
+	UTEST_TRUE_EXPR(BeaconHost->InitHost());
+	BeaconHost->PauseBeaconRequests(false);
+
+	TSharedPtr<BeaconUnitTest::FNetworkStats> BeaconHostNetStats;
+	UTEST_TRUE_EXPR(BeaconUnitTest::ConfigureBeaconNetDriver(*Prerequisites, BeaconHost, &BeaconHostNetStats));
+	UTEST_TRUE_EXPR(BeaconHostNetStats.IsValid());
+
+	// Method 1.
+	// Host handshake calls StartVerifyAuthentication and expects OnAuthenticationVerificationComplete to be called to complete authentication.
+	// Test OnAuthenticationVerificationComplete fired with no delay is working.
+	{
+		BeaconHostNetStats->ReceivedControlMessages.Empty();
+		Prerequisites->GetConfig() = BaseConfig;
+		Prerequisites->GetStats() = BeaconUnitTest::FTestStats();
+
+		Prerequisites->GetConfig().Auth.Method1.bEnabled = true;
+		Prerequisites->GetConfig().Auth.Method1.Result = FOnlineError::Success();
+
+		TObjectPtr<AOnlineBeaconClient> HostUserBeaconClient;
+
+		// Client setup.
+
+		TObjectPtr<AOnlineBeaconUnitTestClient> BeaconClient = Prerequisites->GetWorld()->SpawnActor<AOnlineBeaconUnitTestClient>(AOnlineBeaconUnitTestClient::StaticClass());
+		ON_SCOPE_EXIT{ if (BeaconClient) { BeaconClient->DestroyBeacon(); } };
+		UTEST_TRUE_EXPR(BeaconUnitTest::ConfigureBeacon(*Prerequisites, BeaconClient));
+
+		UTEST_TRUE_EXPR(BeaconClient->GetConnectionState() == EBeaconConnectionState::Invalid);
+		UTEST_TRUE_EXPR(BeaconUnitTest::InitClientForUser(*Prerequisites, BeaconClient, UserId));
+		UTEST_TRUE_EXPR(BeaconClient->GetConnectionState() == EBeaconConnectionState::Pending);
+		UTEST_TRUE_EXPR(BeaconClient->GetNetConnection() != nullptr);
+
+		TSharedPtr<BeaconUnitTest::FNetworkStats> BeaconClientNetStats;
+		UTEST_TRUE_EXPR(BeaconUnitTest::ConfigureBeaconNetDriver(*Prerequisites, BeaconClient, &BeaconClientNetStats));
+		UTEST_TRUE_EXPR(BeaconClientNetStats.IsValid());
+
+		// Handshake testing.
+
+		UTEST_TRUE_EXPR(BeaconUnitTest::TickUntilControlMessageReceived(*Prerequisites, *BeaconHostNetStats));
+		UTEST_TRUE_EXPR(BeaconClient->GetConnectionState() == EBeaconConnectionState::Pending);
+		UTEST_TRUE_EXPR(BeaconClientNetStats->ReceivedControlMessages.IsEmpty());
+		UTEST_TRUE_EXPR(BeaconHostNetStats->ReceivedControlMessages == TArray<uint8>({ NMT_Hello }));
+		UTEST_TRUE_EXPR(BeaconUnitTest::GetBeaconClientForUser(BeaconHostObject, UserId) == nullptr);
+		BeaconHostNetStats->ReceivedControlMessages.Empty();
+
+		UTEST_TRUE_EXPR(BeaconUnitTest::TickOnce(*Prerequisites));
+		UTEST_TRUE_EXPR(BeaconClient->GetConnectionState() == EBeaconConnectionState::Pending);
+		UTEST_TRUE_EXPR(BeaconClientNetStats->ReceivedControlMessages == TArray<uint8>({ NMT_Challenge }));
+		UTEST_TRUE_EXPR(BeaconHostNetStats->ReceivedControlMessages.IsEmpty());
+		BeaconClientNetStats->ReceivedControlMessages.Empty();
+
+		UTEST_TRUE_EXPR(BeaconUnitTest::TickOnce(*Prerequisites));
+		UTEST_TRUE_EXPR(BeaconClient->GetConnectionState() == EBeaconConnectionState::Pending);
+		UTEST_TRUE_EXPR(BeaconClientNetStats->ReceivedControlMessages.IsEmpty());
+		UTEST_TRUE_EXPR(BeaconHostNetStats->ReceivedControlMessages == TArray<uint8>({ NMT_Login }));
+		UTEST_TRUE_EXPR(BeaconUnitTest::GetBeaconClientForUser(BeaconHostObject, UserId) == nullptr);
+		BeaconHostNetStats->ReceivedControlMessages.Empty();
+
+		UTEST_TRUE_EXPR(BeaconUnitTest::TickOnce(*Prerequisites));
+		UTEST_TRUE_EXPR(BeaconClient->GetConnectionState() == EBeaconConnectionState::Pending);
+		UTEST_TRUE_EXPR(BeaconClientNetStats->ReceivedControlMessages == TArray<uint8>({ NMT_BeaconWelcome }));
+		UTEST_TRUE_EXPR(BeaconHostNetStats->ReceivedControlMessages.IsEmpty());
+		BeaconClientNetStats->ReceivedControlMessages.Empty();
+
+		UTEST_TRUE_EXPR(BeaconUnitTest::TickOnce(*Prerequisites));
+		UTEST_TRUE_EXPR(BeaconClient->GetConnectionState() == EBeaconConnectionState::Pending);
+		UTEST_TRUE_EXPR(BeaconClientNetStats->ReceivedControlMessages.IsEmpty());
+		UTEST_TRUE_EXPR(BeaconHostNetStats->ReceivedControlMessages == TArray<uint8>({ NMT_Netspeed, NMT_BeaconJoin }));
+		UTEST_TRUE_EXPR(BeaconUnitTest::GetBeaconClientForUser(BeaconHostObject, UserId) == nullptr);
+		BeaconHostNetStats->ReceivedControlMessages.Empty();
+
+		UTEST_TRUE_EXPR(BeaconUnitTest::TickOnce(*Prerequisites));
+		UTEST_TRUE_EXPR(BeaconClient->GetConnectionState() == EBeaconConnectionState::Pending);
+		UTEST_TRUE_EXPR(BeaconClientNetStats->ReceivedControlMessages == TArray<uint8>({ NMT_BeaconAssignGUID }));
+		UTEST_TRUE_EXPR(BeaconHostNetStats->ReceivedControlMessages.IsEmpty());
+		BeaconClientNetStats->ReceivedControlMessages.Empty();
+
+		UTEST_TRUE_EXPR(BeaconUnitTest::TickOnce(*Prerequisites));
+		UTEST_TRUE_EXPR(BeaconClient->GetConnectionState() == EBeaconConnectionState::Pending);
+		UTEST_TRUE_EXPR(BeaconClientNetStats->ReceivedControlMessages.IsEmpty());
+		UTEST_TRUE_EXPR(BeaconHostNetStats->ReceivedControlMessages == TArray<uint8>({ NMT_BeaconNetGUIDAck }));
+		HostUserBeaconClient = BeaconUnitTest::GetBeaconClientForUser(BeaconHostObject, UserId);
+		UTEST_TRUE_EXPR(HostUserBeaconClient != nullptr);
+		UTEST_TRUE_EXPR(HostUserBeaconClient->GetUniqueId() == *UserId);
+		BeaconHostNetStats->ReceivedControlMessages.Empty();
+
+		// Receive OnConnected RPC on the client.
+		UTEST_TRUE_EXPR(BeaconUnitTest::TickOnce(*Prerequisites));
+		UTEST_TRUE_EXPR(BeaconClient->GetConnectionState() == EBeaconConnectionState::Open);
+		UTEST_TRUE_EXPR(BeaconClientNetStats->ReceivedControlMessages.IsEmpty());
+		UTEST_TRUE_EXPR(BeaconHostNetStats->ReceivedControlMessages.IsEmpty());
+
+		// Handshake complete, channel is open.
+
+		// Disconnect the client.
+		BeaconClient->DestroyBeacon();
+		BeaconClient = nullptr;
+
+		// Check that client actor on the host cleaned up.
+		UTEST_TRUE_EXPR(BeaconUnitTest::TickUntil(*Prerequisites, [&Prerequisites](){ return Prerequisites->GetStats().HostObject.NotifyClientDisconnected.InvokeCount == 1; }));
+		HostUserBeaconClient = BeaconUnitTest::GetBeaconClientForUser(BeaconHostObject, UserId);
+		UTEST_TRUE_EXPR(HostUserBeaconClient == nullptr);
+	}
+
+	// Method 1.
+	// Host handshake calls StartVerifyAuthentication and expects OnAuthenticationVerificationComplete to be called to complete authentication.
+	// Test OnAuthenticationVerificationComplete fired with a single frame delay is working.
+	{
+		BeaconHostNetStats->ReceivedControlMessages.Empty();
+		Prerequisites->GetConfig() = BaseConfig;
+		Prerequisites->GetStats() = BeaconUnitTest::FTestStats();
+
+		Prerequisites->GetConfig().Auth.Method1.bEnabled = true;
+		Prerequisites->GetConfig().Auth.Method1.bDelayDelegate = true;
+		Prerequisites->GetConfig().Auth.Method1.Result = FOnlineError::Success();
+
+		TObjectPtr<AOnlineBeaconClient> HostUserBeaconClient;
+
+		// Client setup.
+
+		TObjectPtr<AOnlineBeaconUnitTestClient> BeaconClient = Prerequisites->GetWorld()->SpawnActor<AOnlineBeaconUnitTestClient>(AOnlineBeaconUnitTestClient::StaticClass());
+		ON_SCOPE_EXIT{ if (BeaconClient) { BeaconClient->DestroyBeacon(); } };
+		UTEST_TRUE_EXPR(BeaconUnitTest::ConfigureBeacon(*Prerequisites, BeaconClient));
+
+		UTEST_TRUE_EXPR(BeaconClient->GetConnectionState() == EBeaconConnectionState::Invalid);
+		UTEST_TRUE_EXPR(BeaconUnitTest::InitClientForUser(*Prerequisites, BeaconClient, UserId));
+		UTEST_TRUE_EXPR(BeaconClient->GetConnectionState() == EBeaconConnectionState::Pending);
+		UTEST_TRUE_EXPR(BeaconClient->GetNetConnection() != nullptr);
+
+		TSharedPtr<BeaconUnitTest::FNetworkStats> BeaconClientNetStats;
+		UTEST_TRUE_EXPR(BeaconUnitTest::ConfigureBeaconNetDriver(*Prerequisites, BeaconClient, &BeaconClientNetStats));
+		UTEST_TRUE_EXPR(BeaconClientNetStats.IsValid());
+
+		// Handshake testing.
+
+		UTEST_TRUE_EXPR(BeaconUnitTest::TickUntilControlMessageReceived(*Prerequisites, *BeaconHostNetStats));
+		UTEST_TRUE_EXPR(BeaconClient->GetConnectionState() == EBeaconConnectionState::Pending);
+		UTEST_TRUE_EXPR(BeaconClientNetStats->ReceivedControlMessages.IsEmpty());
+		UTEST_TRUE_EXPR(BeaconHostNetStats->ReceivedControlMessages == TArray<uint8>({ NMT_Hello }));
+		UTEST_TRUE_EXPR(BeaconUnitTest::GetBeaconClientForUser(BeaconHostObject, UserId) == nullptr);
+		BeaconHostNetStats->ReceivedControlMessages.Empty();
+
+		UTEST_TRUE_EXPR(BeaconUnitTest::TickOnce(*Prerequisites));
+		UTEST_TRUE_EXPR(BeaconClient->GetConnectionState() == EBeaconConnectionState::Pending);
+		UTEST_TRUE_EXPR(BeaconClientNetStats->ReceivedControlMessages == TArray<uint8>({ NMT_Challenge }));
+		UTEST_TRUE_EXPR(BeaconHostNetStats->ReceivedControlMessages.IsEmpty());
+		BeaconClientNetStats->ReceivedControlMessages.Empty();
+
+		UTEST_TRUE_EXPR(BeaconUnitTest::TickOnce(*Prerequisites));
+		UTEST_TRUE_EXPR(BeaconClient->GetConnectionState() == EBeaconConnectionState::Pending);
+		UTEST_TRUE_EXPR(BeaconClientNetStats->ReceivedControlMessages.IsEmpty());
+		UTEST_TRUE_EXPR(BeaconHostNetStats->ReceivedControlMessages == TArray<uint8>({ NMT_Login }));
+		UTEST_TRUE_EXPR(BeaconUnitTest::GetBeaconClientForUser(BeaconHostObject, UserId) == nullptr);
+		BeaconHostNetStats->ReceivedControlMessages.Empty();
+
+		UTEST_TRUE_EXPR(BeaconUnitTest::TickOnce(*Prerequisites));
+		UTEST_TRUE_EXPR(BeaconClient->GetConnectionState() == EBeaconConnectionState::Pending);
+		UTEST_TRUE_EXPR(BeaconClientNetStats->ReceivedControlMessages.IsEmpty());
+		UTEST_TRUE_EXPR(BeaconHostNetStats->ReceivedControlMessages.IsEmpty());
+
+		UTEST_TRUE_EXPR(BeaconUnitTest::TickOnce(*Prerequisites));
+		UTEST_TRUE_EXPR(BeaconClient->GetConnectionState() == EBeaconConnectionState::Pending);
+		UTEST_TRUE_EXPR(BeaconClientNetStats->ReceivedControlMessages == TArray<uint8>({ NMT_BeaconWelcome }));
+		UTEST_TRUE_EXPR(BeaconHostNetStats->ReceivedControlMessages.IsEmpty());
+		BeaconClientNetStats->ReceivedControlMessages.Empty();
+
+		UTEST_TRUE_EXPR(BeaconUnitTest::TickOnce(*Prerequisites));
+		UTEST_TRUE_EXPR(BeaconClient->GetConnectionState() == EBeaconConnectionState::Pending);
+		UTEST_TRUE_EXPR(BeaconClientNetStats->ReceivedControlMessages.IsEmpty());
+		UTEST_TRUE_EXPR(BeaconHostNetStats->ReceivedControlMessages == TArray<uint8>({ NMT_Netspeed, NMT_BeaconJoin }));
+		UTEST_TRUE_EXPR(BeaconUnitTest::GetBeaconClientForUser(BeaconHostObject, UserId) == nullptr);
+		BeaconHostNetStats->ReceivedControlMessages.Empty();
+
+		UTEST_TRUE_EXPR(BeaconUnitTest::TickOnce(*Prerequisites));
+		UTEST_TRUE_EXPR(BeaconClient->GetConnectionState() == EBeaconConnectionState::Pending);
+		UTEST_TRUE_EXPR(BeaconClientNetStats->ReceivedControlMessages == TArray<uint8>({ NMT_BeaconAssignGUID }));
+		UTEST_TRUE_EXPR(BeaconHostNetStats->ReceivedControlMessages.IsEmpty());
+		BeaconClientNetStats->ReceivedControlMessages.Empty();
+
+		UTEST_TRUE_EXPR(BeaconUnitTest::TickOnce(*Prerequisites));
+		UTEST_TRUE_EXPR(BeaconClient->GetConnectionState() == EBeaconConnectionState::Pending);
+		UTEST_TRUE_EXPR(BeaconClientNetStats->ReceivedControlMessages.IsEmpty());
+		UTEST_TRUE_EXPR(BeaconHostNetStats->ReceivedControlMessages == TArray<uint8>({ NMT_BeaconNetGUIDAck }));
+		HostUserBeaconClient = BeaconUnitTest::GetBeaconClientForUser(BeaconHostObject, UserId);
+		UTEST_TRUE_EXPR(HostUserBeaconClient != nullptr);
+		UTEST_TRUE_EXPR(HostUserBeaconClient->GetUniqueId() == *UserId);
+		BeaconHostNetStats->ReceivedControlMessages.Empty();
+
+		// Receive OnConnected RPC on the client.
+		UTEST_TRUE_EXPR(BeaconUnitTest::TickOnce(*Prerequisites));
+		UTEST_TRUE_EXPR(BeaconClient->GetConnectionState() == EBeaconConnectionState::Open);
+		UTEST_TRUE_EXPR(BeaconClientNetStats->ReceivedControlMessages.IsEmpty());
+		UTEST_TRUE_EXPR(BeaconHostNetStats->ReceivedControlMessages.IsEmpty());
+
+		// Handshake complete, channel is open.
+
+		// Disconnect the client.
+		BeaconClient->DestroyBeacon();
+		BeaconClient = nullptr;
+
+		// Check that client actor on the host cleaned up.
+		UTEST_TRUE_EXPR(BeaconUnitTest::TickUntil(*Prerequisites, [&Prerequisites](){ return Prerequisites->GetStats().HostObject.NotifyClientDisconnected.InvokeCount == 1; }));
+		HostUserBeaconClient = BeaconUnitTest::GetBeaconClientForUser(BeaconHostObject, UserId);
+		UTEST_TRUE_EXPR(HostUserBeaconClient == nullptr);
+	}
+	
+	// Method 2.
+	// Host handshake calls StartVerifyAuthentication (callback signature) and expects callback delegate to be called to complete authentication.
+	// Test callback delegate fired with no delay is working.
+	{
+		BeaconHostNetStats->ReceivedControlMessages.Empty();
+		Prerequisites->GetConfig() = BaseConfig;
+		Prerequisites->GetStats() = BeaconUnitTest::FTestStats();
+
+		Prerequisites->GetConfig().Auth.Method2.bEnabled = true;
+		Prerequisites->GetConfig().Auth.Method2.Result = FOnlineError::Success();
+
+		TObjectPtr<AOnlineBeaconClient> HostUserBeaconClient;
+
+		// Client setup.
+
+		TObjectPtr<AOnlineBeaconUnitTestClient> BeaconClient = Prerequisites->GetWorld()->SpawnActor<AOnlineBeaconUnitTestClient>(AOnlineBeaconUnitTestClient::StaticClass());
+		ON_SCOPE_EXIT{ if (BeaconClient) { BeaconClient->DestroyBeacon(); } };
+		UTEST_TRUE_EXPR(BeaconUnitTest::ConfigureBeacon(*Prerequisites, BeaconClient));
+
+		UTEST_TRUE_EXPR(BeaconClient->GetConnectionState() == EBeaconConnectionState::Invalid);
+		UTEST_TRUE_EXPR(BeaconUnitTest::InitClientForUser(*Prerequisites, BeaconClient, UserId));
+		UTEST_TRUE_EXPR(BeaconClient->GetConnectionState() == EBeaconConnectionState::Pending);
+		UTEST_TRUE_EXPR(BeaconClient->GetNetConnection() != nullptr);
+
+		TSharedPtr<BeaconUnitTest::FNetworkStats> BeaconClientNetStats;
+		UTEST_TRUE_EXPR(BeaconUnitTest::ConfigureBeaconNetDriver(*Prerequisites, BeaconClient, &BeaconClientNetStats));
+		UTEST_TRUE_EXPR(BeaconClientNetStats.IsValid());
+
+		// Handshake testing.
+
+		UTEST_TRUE_EXPR(BeaconUnitTest::TickUntilControlMessageReceived(*Prerequisites, *BeaconHostNetStats));
+		UTEST_TRUE_EXPR(BeaconClient->GetConnectionState() == EBeaconConnectionState::Pending);
+		UTEST_TRUE_EXPR(BeaconClientNetStats->ReceivedControlMessages.IsEmpty());
+		UTEST_TRUE_EXPR(BeaconHostNetStats->ReceivedControlMessages == TArray<uint8>({ NMT_Hello }));
+		UTEST_TRUE_EXPR(BeaconUnitTest::GetBeaconClientForUser(BeaconHostObject, UserId) == nullptr);
+		BeaconHostNetStats->ReceivedControlMessages.Empty();
+
+		UTEST_TRUE_EXPR(BeaconUnitTest::TickOnce(*Prerequisites));
+		UTEST_TRUE_EXPR(BeaconClient->GetConnectionState() == EBeaconConnectionState::Pending);
+		UTEST_TRUE_EXPR(BeaconClientNetStats->ReceivedControlMessages == TArray<uint8>({ NMT_Challenge }));
+		UTEST_TRUE_EXPR(BeaconHostNetStats->ReceivedControlMessages.IsEmpty());
+		BeaconClientNetStats->ReceivedControlMessages.Empty();
+
+		UTEST_TRUE_EXPR(BeaconUnitTest::TickOnce(*Prerequisites));
+		UTEST_TRUE_EXPR(BeaconClient->GetConnectionState() == EBeaconConnectionState::Pending);
+		UTEST_TRUE_EXPR(BeaconClientNetStats->ReceivedControlMessages.IsEmpty());
+		UTEST_TRUE_EXPR(BeaconHostNetStats->ReceivedControlMessages == TArray<uint8>({ NMT_Login }));
+		UTEST_TRUE_EXPR(BeaconUnitTest::GetBeaconClientForUser(BeaconHostObject, UserId) == nullptr);
+		BeaconHostNetStats->ReceivedControlMessages.Empty();
+
+		UTEST_TRUE_EXPR(BeaconUnitTest::TickOnce(*Prerequisites));
+		UTEST_TRUE_EXPR(BeaconClient->GetConnectionState() == EBeaconConnectionState::Pending);
+		UTEST_TRUE_EXPR(BeaconClientNetStats->ReceivedControlMessages == TArray<uint8>({ NMT_BeaconWelcome }));
+		UTEST_TRUE_EXPR(BeaconHostNetStats->ReceivedControlMessages.IsEmpty());
+		BeaconClientNetStats->ReceivedControlMessages.Empty();
+
+		UTEST_TRUE_EXPR(BeaconUnitTest::TickOnce(*Prerequisites));
+		UTEST_TRUE_EXPR(BeaconClient->GetConnectionState() == EBeaconConnectionState::Pending);
+		UTEST_TRUE_EXPR(BeaconClientNetStats->ReceivedControlMessages.IsEmpty());
+		UTEST_TRUE_EXPR(BeaconHostNetStats->ReceivedControlMessages == TArray<uint8>({ NMT_Netspeed, NMT_BeaconJoin }));
+		UTEST_TRUE_EXPR(BeaconUnitTest::GetBeaconClientForUser(BeaconHostObject, UserId) == nullptr);
+		BeaconHostNetStats->ReceivedControlMessages.Empty();
+
+		UTEST_TRUE_EXPR(BeaconUnitTest::TickOnce(*Prerequisites));
+		UTEST_TRUE_EXPR(BeaconClient->GetConnectionState() == EBeaconConnectionState::Pending);
+		UTEST_TRUE_EXPR(BeaconClientNetStats->ReceivedControlMessages == TArray<uint8>({ NMT_BeaconAssignGUID }));
+		UTEST_TRUE_EXPR(BeaconHostNetStats->ReceivedControlMessages.IsEmpty());
+		BeaconClientNetStats->ReceivedControlMessages.Empty();
+
+		UTEST_TRUE_EXPR(BeaconUnitTest::TickOnce(*Prerequisites));
+		UTEST_TRUE_EXPR(BeaconClient->GetConnectionState() == EBeaconConnectionState::Pending);
+		UTEST_TRUE_EXPR(BeaconClientNetStats->ReceivedControlMessages.IsEmpty());
+		UTEST_TRUE_EXPR(BeaconHostNetStats->ReceivedControlMessages == TArray<uint8>({ NMT_BeaconNetGUIDAck }));
+		HostUserBeaconClient = BeaconUnitTest::GetBeaconClientForUser(BeaconHostObject, UserId);
+		UTEST_TRUE_EXPR(HostUserBeaconClient != nullptr);
+		UTEST_TRUE_EXPR(HostUserBeaconClient->GetUniqueId() == *UserId);
+		BeaconHostNetStats->ReceivedControlMessages.Empty();
+
+		// Receive OnConnected RPC on the client.
+		UTEST_TRUE_EXPR(BeaconUnitTest::TickOnce(*Prerequisites));
+		UTEST_TRUE_EXPR(BeaconClient->GetConnectionState() == EBeaconConnectionState::Open);
+		UTEST_TRUE_EXPR(BeaconClientNetStats->ReceivedControlMessages.IsEmpty());
+		UTEST_TRUE_EXPR(BeaconHostNetStats->ReceivedControlMessages.IsEmpty());
+
+		// Handshake complete, channel is open.
+
+		// Disconnect the client.
+		BeaconClient->DestroyBeacon();
+		BeaconClient = nullptr;
+
+		// Check that client actor on the host cleaned up.
+		UTEST_TRUE_EXPR(BeaconUnitTest::TickUntil(*Prerequisites, [&Prerequisites](){ return Prerequisites->GetStats().HostObject.NotifyClientDisconnected.InvokeCount == 1; }));
+		HostUserBeaconClient = BeaconUnitTest::GetBeaconClientForUser(BeaconHostObject, UserId);
+		UTEST_TRUE_EXPR(HostUserBeaconClient == nullptr);
+	}
+
+	// Method 2.
+	// Host handshake calls StartVerifyAuthentication (callback signature) and expects callback delegate to be called to complete authentication.
+	// Test callback delegate fired with a single frame delay is working.
+	{
+		BeaconHostNetStats->ReceivedControlMessages.Empty();
+		Prerequisites->GetConfig() = BaseConfig;
+		Prerequisites->GetStats() = BeaconUnitTest::FTestStats();
+
+		Prerequisites->GetConfig().Auth.Method2.bEnabled = true;
+		Prerequisites->GetConfig().Auth.Method2.bDelayDelegate = true;
+		Prerequisites->GetConfig().Auth.Method2.Result = FOnlineError::Success();
+
+		TObjectPtr<AOnlineBeaconClient> HostUserBeaconClient;
+
+		// Client setup.
+
+		TObjectPtr<AOnlineBeaconUnitTestClient> BeaconClient = Prerequisites->GetWorld()->SpawnActor<AOnlineBeaconUnitTestClient>(AOnlineBeaconUnitTestClient::StaticClass());
+		ON_SCOPE_EXIT{ if (BeaconClient) { BeaconClient->DestroyBeacon(); } };
+		UTEST_TRUE_EXPR(BeaconUnitTest::ConfigureBeacon(*Prerequisites, BeaconClient));
+
+		UTEST_TRUE_EXPR(BeaconClient->GetConnectionState() == EBeaconConnectionState::Invalid);
+		UTEST_TRUE_EXPR(BeaconUnitTest::InitClientForUser(*Prerequisites, BeaconClient, UserId));
+		UTEST_TRUE_EXPR(BeaconClient->GetConnectionState() == EBeaconConnectionState::Pending);
+		UTEST_TRUE_EXPR(BeaconClient->GetNetConnection() != nullptr);
+
+		TSharedPtr<BeaconUnitTest::FNetworkStats> BeaconClientNetStats;
+		UTEST_TRUE_EXPR(BeaconUnitTest::ConfigureBeaconNetDriver(*Prerequisites, BeaconClient, &BeaconClientNetStats));
+		UTEST_TRUE_EXPR(BeaconClientNetStats.IsValid());
+
+		// Handshake testing.
+
+		UTEST_TRUE_EXPR(BeaconUnitTest::TickUntilControlMessageReceived(*Prerequisites, *BeaconHostNetStats));
+		UTEST_TRUE_EXPR(BeaconClient->GetConnectionState() == EBeaconConnectionState::Pending);
+		UTEST_TRUE_EXPR(BeaconClientNetStats->ReceivedControlMessages.IsEmpty());
+		UTEST_TRUE_EXPR(BeaconHostNetStats->ReceivedControlMessages == TArray<uint8>({ NMT_Hello }));
+		UTEST_TRUE_EXPR(BeaconUnitTest::GetBeaconClientForUser(BeaconHostObject, UserId) == nullptr);
+		BeaconHostNetStats->ReceivedControlMessages.Empty();
+
+		UTEST_TRUE_EXPR(BeaconUnitTest::TickOnce(*Prerequisites));
+		UTEST_TRUE_EXPR(BeaconClient->GetConnectionState() == EBeaconConnectionState::Pending);
+		UTEST_TRUE_EXPR(BeaconClientNetStats->ReceivedControlMessages == TArray<uint8>({ NMT_Challenge }));
+		UTEST_TRUE_EXPR(BeaconHostNetStats->ReceivedControlMessages.IsEmpty());
+		BeaconClientNetStats->ReceivedControlMessages.Empty();
+
+		UTEST_TRUE_EXPR(BeaconUnitTest::TickOnce(*Prerequisites));
+		UTEST_TRUE_EXPR(BeaconClient->GetConnectionState() == EBeaconConnectionState::Pending);
+		UTEST_TRUE_EXPR(BeaconClientNetStats->ReceivedControlMessages.IsEmpty());
+		UTEST_TRUE_EXPR(BeaconHostNetStats->ReceivedControlMessages == TArray<uint8>({ NMT_Login }));
+		UTEST_TRUE_EXPR(BeaconUnitTest::GetBeaconClientForUser(BeaconHostObject, UserId) == nullptr);
+		BeaconHostNetStats->ReceivedControlMessages.Empty();
+
+		UTEST_TRUE_EXPR(BeaconUnitTest::TickOnce(*Prerequisites));
+		UTEST_TRUE_EXPR(BeaconClient->GetConnectionState() == EBeaconConnectionState::Pending);
+		UTEST_TRUE_EXPR(BeaconClientNetStats->ReceivedControlMessages.IsEmpty());
+		UTEST_TRUE_EXPR(BeaconHostNetStats->ReceivedControlMessages.IsEmpty());
+
+		UTEST_TRUE_EXPR(BeaconUnitTest::TickOnce(*Prerequisites));
+		UTEST_TRUE_EXPR(BeaconClient->GetConnectionState() == EBeaconConnectionState::Pending);
+		UTEST_TRUE_EXPR(BeaconClientNetStats->ReceivedControlMessages == TArray<uint8>({ NMT_BeaconWelcome }));
+		UTEST_TRUE_EXPR(BeaconHostNetStats->ReceivedControlMessages.IsEmpty());
+		BeaconClientNetStats->ReceivedControlMessages.Empty();
+
+		UTEST_TRUE_EXPR(BeaconUnitTest::TickOnce(*Prerequisites));
+		UTEST_TRUE_EXPR(BeaconClient->GetConnectionState() == EBeaconConnectionState::Pending);
+		UTEST_TRUE_EXPR(BeaconClientNetStats->ReceivedControlMessages.IsEmpty());
+		UTEST_TRUE_EXPR(BeaconHostNetStats->ReceivedControlMessages == TArray<uint8>({ NMT_Netspeed, NMT_BeaconJoin }));
+		UTEST_TRUE_EXPR(BeaconUnitTest::GetBeaconClientForUser(BeaconHostObject, UserId) == nullptr);
+		BeaconHostNetStats->ReceivedControlMessages.Empty();
+
+		UTEST_TRUE_EXPR(BeaconUnitTest::TickOnce(*Prerequisites));
+		UTEST_TRUE_EXPR(BeaconClient->GetConnectionState() == EBeaconConnectionState::Pending);
+		UTEST_TRUE_EXPR(BeaconClientNetStats->ReceivedControlMessages == TArray<uint8>({ NMT_BeaconAssignGUID }));
+		UTEST_TRUE_EXPR(BeaconHostNetStats->ReceivedControlMessages.IsEmpty());
+		BeaconClientNetStats->ReceivedControlMessages.Empty();
+
+		UTEST_TRUE_EXPR(BeaconUnitTest::TickOnce(*Prerequisites));
+		UTEST_TRUE_EXPR(BeaconClient->GetConnectionState() == EBeaconConnectionState::Pending);
+		UTEST_TRUE_EXPR(BeaconClientNetStats->ReceivedControlMessages.IsEmpty());
+		UTEST_TRUE_EXPR(BeaconHostNetStats->ReceivedControlMessages == TArray<uint8>({ NMT_BeaconNetGUIDAck }));
+		HostUserBeaconClient = BeaconUnitTest::GetBeaconClientForUser(BeaconHostObject, UserId);
+		UTEST_TRUE_EXPR(HostUserBeaconClient != nullptr);
+		UTEST_TRUE_EXPR(HostUserBeaconClient->GetUniqueId() == *UserId);
+		BeaconHostNetStats->ReceivedControlMessages.Empty();
+
+		// Receive OnConnected RPC on the client.
+		UTEST_TRUE_EXPR(BeaconUnitTest::TickOnce(*Prerequisites));
+		UTEST_TRUE_EXPR(BeaconClient->GetConnectionState() == EBeaconConnectionState::Open);
+		UTEST_TRUE_EXPR(BeaconClientNetStats->ReceivedControlMessages.IsEmpty());
+		UTEST_TRUE_EXPR(BeaconHostNetStats->ReceivedControlMessages.IsEmpty());
+
+		// Handshake complete, channel is open.
+
+		// Disconnect the client.
+		BeaconClient->DestroyBeacon();
+		BeaconClient = nullptr;
+
+		// Check that client actor on the host cleaned up.
+		UTEST_TRUE_EXPR(BeaconUnitTest::TickUntil(*Prerequisites, [&Prerequisites](){ return Prerequisites->GetStats().HostObject.NotifyClientDisconnected.InvokeCount == 1; }));
+		HostUserBeaconClient = BeaconUnitTest::GetBeaconClientForUser(BeaconHostObject, UserId);
+		UTEST_TRUE_EXPR(HostUserBeaconClient == nullptr);
+	}
+
+	// Method 3.
+	// Host handshake calls StartVerifyAuthentication (callback signature with LoginOptions) and expects callback delegate to be called to complete authentication.
+	// Test callback delegate fired with no delay is working.
+	{
+		BeaconHostNetStats->ReceivedControlMessages.Empty();
+		Prerequisites->GetConfig() = BaseConfig;
+		Prerequisites->GetStats() = BeaconUnitTest::FTestStats();
+
+		Prerequisites->GetConfig().Auth.Method3.bEnabled = true;
+		Prerequisites->GetConfig().Auth.Method3.Result = FOnlineError::Success();
+
+		TObjectPtr<AOnlineBeaconClient> HostUserBeaconClient;
+
+		// Client setup.
+
+		TObjectPtr<AOnlineBeaconUnitTestClient> BeaconClient = Prerequisites->GetWorld()->SpawnActor<AOnlineBeaconUnitTestClient>(AOnlineBeaconUnitTestClient::StaticClass());
+		ON_SCOPE_EXIT{ if (BeaconClient) { BeaconClient->DestroyBeacon(); } };
+		UTEST_TRUE_EXPR(BeaconUnitTest::ConfigureBeacon(*Prerequisites, BeaconClient));
+
+		UTEST_TRUE_EXPR(BeaconClient->GetConnectionState() == EBeaconConnectionState::Invalid);
+		UTEST_TRUE_EXPR(BeaconUnitTest::InitClientForUser(*Prerequisites, BeaconClient, UserId));
+		UTEST_TRUE_EXPR(BeaconClient->GetConnectionState() == EBeaconConnectionState::Pending);
+		UTEST_TRUE_EXPR(BeaconClient->GetNetConnection() != nullptr);
+
+		TSharedPtr<BeaconUnitTest::FNetworkStats> BeaconClientNetStats;
+		UTEST_TRUE_EXPR(BeaconUnitTest::ConfigureBeaconNetDriver(*Prerequisites, BeaconClient, &BeaconClientNetStats));
+		UTEST_TRUE_EXPR(BeaconClientNetStats.IsValid());
+
+		// Handshake testing.
+
+		UTEST_TRUE_EXPR(BeaconUnitTest::TickUntilControlMessageReceived(*Prerequisites, *BeaconHostNetStats));
+		UTEST_TRUE_EXPR(BeaconClient->GetConnectionState() == EBeaconConnectionState::Pending);
+		UTEST_TRUE_EXPR(BeaconClientNetStats->ReceivedControlMessages.IsEmpty());
+		UTEST_TRUE_EXPR(BeaconHostNetStats->ReceivedControlMessages == TArray<uint8>({ NMT_Hello }));
+		UTEST_TRUE_EXPR(BeaconUnitTest::GetBeaconClientForUser(BeaconHostObject, UserId) == nullptr);
+		BeaconHostNetStats->ReceivedControlMessages.Empty();
+
+		UTEST_TRUE_EXPR(BeaconUnitTest::TickOnce(*Prerequisites));
+		UTEST_TRUE_EXPR(BeaconClient->GetConnectionState() == EBeaconConnectionState::Pending);
+		UTEST_TRUE_EXPR(BeaconClientNetStats->ReceivedControlMessages == TArray<uint8>({ NMT_Challenge }));
+		UTEST_TRUE_EXPR(BeaconHostNetStats->ReceivedControlMessages.IsEmpty());
+		BeaconClientNetStats->ReceivedControlMessages.Empty();
+
+		UTEST_TRUE_EXPR(BeaconUnitTest::TickOnce(*Prerequisites));
+		UTEST_TRUE_EXPR(BeaconClient->GetConnectionState() == EBeaconConnectionState::Pending);
+		UTEST_TRUE_EXPR(BeaconClientNetStats->ReceivedControlMessages.IsEmpty());
+		UTEST_TRUE_EXPR(BeaconHostNetStats->ReceivedControlMessages == TArray<uint8>({ NMT_Login }));
+		UTEST_TRUE_EXPR(BeaconUnitTest::GetBeaconClientForUser(BeaconHostObject, UserId) == nullptr);
+		BeaconHostNetStats->ReceivedControlMessages.Empty();
+
+		UTEST_TRUE_EXPR(BeaconUnitTest::TickOnce(*Prerequisites));
+		UTEST_TRUE_EXPR(BeaconClient->GetConnectionState() == EBeaconConnectionState::Pending);
+		UTEST_TRUE_EXPR(BeaconClientNetStats->ReceivedControlMessages == TArray<uint8>({ NMT_BeaconWelcome }));
+		UTEST_TRUE_EXPR(BeaconHostNetStats->ReceivedControlMessages.IsEmpty());
+		BeaconClientNetStats->ReceivedControlMessages.Empty();
+
+		UTEST_TRUE_EXPR(BeaconUnitTest::TickOnce(*Prerequisites));
+		UTEST_TRUE_EXPR(BeaconClient->GetConnectionState() == EBeaconConnectionState::Pending);
+		UTEST_TRUE_EXPR(BeaconClientNetStats->ReceivedControlMessages.IsEmpty());
+		UTEST_TRUE_EXPR(BeaconHostNetStats->ReceivedControlMessages == TArray<uint8>({ NMT_Netspeed, NMT_BeaconJoin }));
+		UTEST_TRUE_EXPR(BeaconUnitTest::GetBeaconClientForUser(BeaconHostObject, UserId) == nullptr);
+		BeaconHostNetStats->ReceivedControlMessages.Empty();
+
+		UTEST_TRUE_EXPR(BeaconUnitTest::TickOnce(*Prerequisites));
+		UTEST_TRUE_EXPR(BeaconClient->GetConnectionState() == EBeaconConnectionState::Pending);
+		UTEST_TRUE_EXPR(BeaconClientNetStats->ReceivedControlMessages == TArray<uint8>({ NMT_BeaconAssignGUID }));
+		UTEST_TRUE_EXPR(BeaconHostNetStats->ReceivedControlMessages.IsEmpty());
+		BeaconClientNetStats->ReceivedControlMessages.Empty();
+
+		UTEST_TRUE_EXPR(BeaconUnitTest::TickOnce(*Prerequisites));
+		UTEST_TRUE_EXPR(BeaconClient->GetConnectionState() == EBeaconConnectionState::Pending);
+		UTEST_TRUE_EXPR(BeaconClientNetStats->ReceivedControlMessages.IsEmpty());
+		UTEST_TRUE_EXPR(BeaconHostNetStats->ReceivedControlMessages == TArray<uint8>({ NMT_BeaconNetGUIDAck }));
+		HostUserBeaconClient = BeaconUnitTest::GetBeaconClientForUser(BeaconHostObject, UserId);
+		UTEST_TRUE_EXPR(HostUserBeaconClient != nullptr);
+		UTEST_TRUE_EXPR(HostUserBeaconClient->GetUniqueId() == *UserId);
+		BeaconHostNetStats->ReceivedControlMessages.Empty();
+
+		// Receive OnConnected RPC on the client.
+		UTEST_TRUE_EXPR(BeaconUnitTest::TickOnce(*Prerequisites));
+		UTEST_TRUE_EXPR(BeaconClient->GetConnectionState() == EBeaconConnectionState::Open);
+		UTEST_TRUE_EXPR(BeaconClientNetStats->ReceivedControlMessages.IsEmpty());
+		UTEST_TRUE_EXPR(BeaconHostNetStats->ReceivedControlMessages.IsEmpty());
+
+		// Handshake complete, channel is open.
+
+		// Disconnect the client.
+		BeaconClient->DestroyBeacon();
+		BeaconClient = nullptr;
+
+		// Check that client actor on the host cleaned up.
+		UTEST_TRUE_EXPR(BeaconUnitTest::TickUntil(*Prerequisites, [&Prerequisites](){ return Prerequisites->GetStats().HostObject.NotifyClientDisconnected.InvokeCount == 1; }));
+		HostUserBeaconClient = BeaconUnitTest::GetBeaconClientForUser(BeaconHostObject, UserId);
+		UTEST_TRUE_EXPR(HostUserBeaconClient == nullptr);
+	}
+
+	// Method 3.
+	// Host handshake calls StartVerifyAuthentication (callback signature with LoginOptions) and expects callback delegate to be called to complete authentication.
+	// Test callback delegate fired with a single frame delay is working.
+	{
+		BeaconHostNetStats->ReceivedControlMessages.Empty();
+		Prerequisites->GetConfig() = BaseConfig;
+		Prerequisites->GetStats() = BeaconUnitTest::FTestStats();
+
+		Prerequisites->GetConfig().Auth.Method3.bEnabled = true;
+		Prerequisites->GetConfig().Auth.Method3.bDelayDelegate = true;
+		Prerequisites->GetConfig().Auth.Method3.Result = FOnlineError::Success();
+
+		TObjectPtr<AOnlineBeaconClient> HostUserBeaconClient;
+
+		// Client setup.
+
+		TObjectPtr<AOnlineBeaconUnitTestClient> BeaconClient = Prerequisites->GetWorld()->SpawnActor<AOnlineBeaconUnitTestClient>(AOnlineBeaconUnitTestClient::StaticClass());
+		ON_SCOPE_EXIT{ if (BeaconClient) { BeaconClient->DestroyBeacon(); } };
+		UTEST_TRUE_EXPR(BeaconUnitTest::ConfigureBeacon(*Prerequisites, BeaconClient));
+
+		UTEST_TRUE_EXPR(BeaconClient->GetConnectionState() == EBeaconConnectionState::Invalid);
+		UTEST_TRUE_EXPR(BeaconUnitTest::InitClientForUser(*Prerequisites, BeaconClient, UserId));
+		UTEST_TRUE_EXPR(BeaconClient->GetConnectionState() == EBeaconConnectionState::Pending);
+		UTEST_TRUE_EXPR(BeaconClient->GetNetConnection() != nullptr);
+
+		TSharedPtr<BeaconUnitTest::FNetworkStats> BeaconClientNetStats;
+		UTEST_TRUE_EXPR(BeaconUnitTest::ConfigureBeaconNetDriver(*Prerequisites, BeaconClient, &BeaconClientNetStats));
+		UTEST_TRUE_EXPR(BeaconClientNetStats.IsValid());
+
+		// Handshake testing.
+
+		UTEST_TRUE_EXPR(BeaconUnitTest::TickUntilControlMessageReceived(*Prerequisites, *BeaconHostNetStats));
+		UTEST_TRUE_EXPR(BeaconClient->GetConnectionState() == EBeaconConnectionState::Pending);
+		UTEST_TRUE_EXPR(BeaconClientNetStats->ReceivedControlMessages.IsEmpty());
+		UTEST_TRUE_EXPR(BeaconHostNetStats->ReceivedControlMessages == TArray<uint8>({ NMT_Hello }));
+		UTEST_TRUE_EXPR(BeaconUnitTest::GetBeaconClientForUser(BeaconHostObject, UserId) == nullptr);
+		BeaconHostNetStats->ReceivedControlMessages.Empty();
+
+		UTEST_TRUE_EXPR(BeaconUnitTest::TickOnce(*Prerequisites));
+		UTEST_TRUE_EXPR(BeaconClient->GetConnectionState() == EBeaconConnectionState::Pending);
+		UTEST_TRUE_EXPR(BeaconClientNetStats->ReceivedControlMessages == TArray<uint8>({ NMT_Challenge }));
+		UTEST_TRUE_EXPR(BeaconHostNetStats->ReceivedControlMessages.IsEmpty());
+		BeaconClientNetStats->ReceivedControlMessages.Empty();
+
+		UTEST_TRUE_EXPR(BeaconUnitTest::TickOnce(*Prerequisites));
+		UTEST_TRUE_EXPR(BeaconClient->GetConnectionState() == EBeaconConnectionState::Pending);
+		UTEST_TRUE_EXPR(BeaconClientNetStats->ReceivedControlMessages.IsEmpty());
+		UTEST_TRUE_EXPR(BeaconHostNetStats->ReceivedControlMessages == TArray<uint8>({ NMT_Login }));
+		UTEST_TRUE_EXPR(BeaconUnitTest::GetBeaconClientForUser(BeaconHostObject, UserId) == nullptr);
+		BeaconHostNetStats->ReceivedControlMessages.Empty();
+
+		UTEST_TRUE_EXPR(BeaconUnitTest::TickOnce(*Prerequisites));
+		UTEST_TRUE_EXPR(BeaconClient->GetConnectionState() == EBeaconConnectionState::Pending);
+		UTEST_TRUE_EXPR(BeaconClientNetStats->ReceivedControlMessages.IsEmpty());
+		UTEST_TRUE_EXPR(BeaconHostNetStats->ReceivedControlMessages.IsEmpty());
+
+		UTEST_TRUE_EXPR(BeaconUnitTest::TickOnce(*Prerequisites));
+		UTEST_TRUE_EXPR(BeaconClient->GetConnectionState() == EBeaconConnectionState::Pending);
+		UTEST_TRUE_EXPR(BeaconClientNetStats->ReceivedControlMessages == TArray<uint8>({ NMT_BeaconWelcome }));
+		UTEST_TRUE_EXPR(BeaconHostNetStats->ReceivedControlMessages.IsEmpty());
+		BeaconClientNetStats->ReceivedControlMessages.Empty();
+
+		UTEST_TRUE_EXPR(BeaconUnitTest::TickOnce(*Prerequisites));
+		UTEST_TRUE_EXPR(BeaconClient->GetConnectionState() == EBeaconConnectionState::Pending);
+		UTEST_TRUE_EXPR(BeaconClientNetStats->ReceivedControlMessages.IsEmpty());
+		UTEST_TRUE_EXPR(BeaconHostNetStats->ReceivedControlMessages == TArray<uint8>({ NMT_Netspeed, NMT_BeaconJoin }));
+		UTEST_TRUE_EXPR(BeaconUnitTest::GetBeaconClientForUser(BeaconHostObject, UserId) == nullptr);
+		BeaconHostNetStats->ReceivedControlMessages.Empty();
+
+		UTEST_TRUE_EXPR(BeaconUnitTest::TickOnce(*Prerequisites));
+		UTEST_TRUE_EXPR(BeaconClient->GetConnectionState() == EBeaconConnectionState::Pending);
+		UTEST_TRUE_EXPR(BeaconClientNetStats->ReceivedControlMessages == TArray<uint8>({ NMT_BeaconAssignGUID }));
+		UTEST_TRUE_EXPR(BeaconHostNetStats->ReceivedControlMessages.IsEmpty());
+		BeaconClientNetStats->ReceivedControlMessages.Empty();
+
+		UTEST_TRUE_EXPR(BeaconUnitTest::TickOnce(*Prerequisites));
+		UTEST_TRUE_EXPR(BeaconClient->GetConnectionState() == EBeaconConnectionState::Pending);
+		UTEST_TRUE_EXPR(BeaconClientNetStats->ReceivedControlMessages.IsEmpty());
+		UTEST_TRUE_EXPR(BeaconHostNetStats->ReceivedControlMessages == TArray<uint8>({ NMT_BeaconNetGUIDAck }));
+		HostUserBeaconClient = BeaconUnitTest::GetBeaconClientForUser(BeaconHostObject, UserId);
+		UTEST_TRUE_EXPR(HostUserBeaconClient != nullptr);
+		UTEST_TRUE_EXPR(HostUserBeaconClient->GetUniqueId() == *UserId);
+		BeaconHostNetStats->ReceivedControlMessages.Empty();
+
+		// Receive OnConnected RPC on the client.
+		UTEST_TRUE_EXPR(BeaconUnitTest::TickOnce(*Prerequisites));
+		UTEST_TRUE_EXPR(BeaconClient->GetConnectionState() == EBeaconConnectionState::Open);
+		UTEST_TRUE_EXPR(BeaconClientNetStats->ReceivedControlMessages.IsEmpty());
+		UTEST_TRUE_EXPR(BeaconHostNetStats->ReceivedControlMessages.IsEmpty());
+
+		// Handshake complete, channel is open.
+
+		// Disconnect the client.
+		BeaconClient->DestroyBeacon();
+		BeaconClient = nullptr;
+
+		// Check that client actor on the host cleaned up.
+		UTEST_TRUE_EXPR(BeaconUnitTest::TickUntil(*Prerequisites, [&Prerequisites](){ return Prerequisites->GetStats().HostObject.NotifyClientDisconnected.InvokeCount == 1; }));
+		HostUserBeaconClient = BeaconUnitTest::GetBeaconClientForUser(BeaconHostObject, UserId);
+		UTEST_TRUE_EXPR(HostUserBeaconClient == nullptr);
+	}
+
+	// Auth verify.
+	// Use Method 3 auth and verify that a user supplied VerifyJoinForBeaconType method allows the connection to proceed.
+	{
+		BeaconHostNetStats->ReceivedControlMessages.Empty();
+		Prerequisites->GetConfig() = BaseConfig;
+		Prerequisites->GetStats() = BeaconUnitTest::FTestStats();
+
+		Prerequisites->GetConfig().Auth.Method3.bEnabled = true;
+		Prerequisites->GetConfig().Auth.Method3.Result = FOnlineError::Success();
+
+		Prerequisites->GetConfig().Auth.Verify.bEnabled = true;
+		Prerequisites->GetConfig().Auth.Verify.bResult = true;
+
+		TObjectPtr<AOnlineBeaconClient> HostUserBeaconClient;
+
+		// Client setup.
+
+		TObjectPtr<AOnlineBeaconUnitTestClient> BeaconClient = Prerequisites->GetWorld()->SpawnActor<AOnlineBeaconUnitTestClient>(AOnlineBeaconUnitTestClient::StaticClass());
+		ON_SCOPE_EXIT{ if (BeaconClient) { BeaconClient->DestroyBeacon(); } };
+		UTEST_TRUE_EXPR(BeaconUnitTest::ConfigureBeacon(*Prerequisites, BeaconClient));
+
+		UTEST_TRUE_EXPR(BeaconClient->GetConnectionState() == EBeaconConnectionState::Invalid);
+		UTEST_TRUE_EXPR(BeaconUnitTest::InitClientForUser(*Prerequisites, BeaconClient, UserId));
+		UTEST_TRUE_EXPR(BeaconClient->GetConnectionState() == EBeaconConnectionState::Pending);
+		UTEST_TRUE_EXPR(BeaconClient->GetNetConnection() != nullptr);
+
+		TSharedPtr<BeaconUnitTest::FNetworkStats> BeaconClientNetStats;
+		UTEST_TRUE_EXPR(BeaconUnitTest::ConfigureBeaconNetDriver(*Prerequisites, BeaconClient, &BeaconClientNetStats));
+		UTEST_TRUE_EXPR(BeaconClientNetStats.IsValid());
+
+		// Handshake testing.
+
+		UTEST_TRUE_EXPR(BeaconUnitTest::TickUntilControlMessageReceived(*Prerequisites, *BeaconHostNetStats));
+		UTEST_TRUE_EXPR(BeaconClient->GetConnectionState() == EBeaconConnectionState::Pending);
+		UTEST_TRUE_EXPR(BeaconClientNetStats->ReceivedControlMessages.IsEmpty());
+		UTEST_TRUE_EXPR(BeaconHostNetStats->ReceivedControlMessages == TArray<uint8>({ NMT_Hello }));
+		UTEST_TRUE_EXPR(BeaconUnitTest::GetBeaconClientForUser(BeaconHostObject, UserId) == nullptr);
+		BeaconHostNetStats->ReceivedControlMessages.Empty();
+
+		UTEST_TRUE_EXPR(BeaconUnitTest::TickOnce(*Prerequisites));
+		UTEST_TRUE_EXPR(BeaconClient->GetConnectionState() == EBeaconConnectionState::Pending);
+		UTEST_TRUE_EXPR(BeaconClientNetStats->ReceivedControlMessages == TArray<uint8>({ NMT_Challenge }));
+		UTEST_TRUE_EXPR(BeaconHostNetStats->ReceivedControlMessages.IsEmpty());
+		BeaconClientNetStats->ReceivedControlMessages.Empty();
+
+		UTEST_TRUE_EXPR(BeaconUnitTest::TickOnce(*Prerequisites));
+		UTEST_TRUE_EXPR(BeaconClient->GetConnectionState() == EBeaconConnectionState::Pending);
+		UTEST_TRUE_EXPR(BeaconClientNetStats->ReceivedControlMessages.IsEmpty());
+		UTEST_TRUE_EXPR(BeaconHostNetStats->ReceivedControlMessages == TArray<uint8>({ NMT_Login }));
+		UTEST_TRUE_EXPR(BeaconUnitTest::GetBeaconClientForUser(BeaconHostObject, UserId) == nullptr);
+		BeaconHostNetStats->ReceivedControlMessages.Empty();
+
+		UTEST_TRUE_EXPR(BeaconUnitTest::TickOnce(*Prerequisites));
+		UTEST_TRUE_EXPR(BeaconClient->GetConnectionState() == EBeaconConnectionState::Pending);
+		UTEST_TRUE_EXPR(BeaconClientNetStats->ReceivedControlMessages == TArray<uint8>({ NMT_BeaconWelcome }));
+		UTEST_TRUE_EXPR(BeaconHostNetStats->ReceivedControlMessages.IsEmpty());
+		BeaconClientNetStats->ReceivedControlMessages.Empty();
+
+		UTEST_TRUE_EXPR(BeaconUnitTest::TickOnce(*Prerequisites));
+		UTEST_TRUE_EXPR(BeaconClient->GetConnectionState() == EBeaconConnectionState::Pending);
+		UTEST_TRUE_EXPR(BeaconClientNetStats->ReceivedControlMessages.IsEmpty());
+		UTEST_TRUE_EXPR(BeaconHostNetStats->ReceivedControlMessages == TArray<uint8>({ NMT_Netspeed, NMT_BeaconJoin }));
+		UTEST_TRUE_EXPR(BeaconUnitTest::GetBeaconClientForUser(BeaconHostObject, UserId) == nullptr);
+		BeaconHostNetStats->ReceivedControlMessages.Empty();
+
+		UTEST_TRUE_EXPR(BeaconUnitTest::TickOnce(*Prerequisites));
+		UTEST_TRUE_EXPR(BeaconClient->GetConnectionState() == EBeaconConnectionState::Pending);
+		UTEST_TRUE_EXPR(BeaconClientNetStats->ReceivedControlMessages == TArray<uint8>({ NMT_BeaconAssignGUID }));
+		UTEST_TRUE_EXPR(BeaconHostNetStats->ReceivedControlMessages.IsEmpty());
+		BeaconClientNetStats->ReceivedControlMessages.Empty();
+
+		UTEST_TRUE_EXPR(BeaconUnitTest::TickOnce(*Prerequisites));
+		UTEST_TRUE_EXPR(BeaconClient->GetConnectionState() == EBeaconConnectionState::Pending);
+		UTEST_TRUE_EXPR(BeaconClientNetStats->ReceivedControlMessages.IsEmpty());
+		UTEST_TRUE_EXPR(BeaconHostNetStats->ReceivedControlMessages == TArray<uint8>({ NMT_BeaconNetGUIDAck }));
+		HostUserBeaconClient = BeaconUnitTest::GetBeaconClientForUser(BeaconHostObject, UserId);
+		UTEST_TRUE_EXPR(HostUserBeaconClient != nullptr);
+		UTEST_TRUE_EXPR(HostUserBeaconClient->GetUniqueId() == *UserId);
+		BeaconHostNetStats->ReceivedControlMessages.Empty();
+
+		// Receive OnConnected RPC on the client.
+		UTEST_TRUE_EXPR(BeaconUnitTest::TickOnce(*Prerequisites));
+		UTEST_TRUE_EXPR(BeaconClient->GetConnectionState() == EBeaconConnectionState::Open);
+		UTEST_TRUE_EXPR(BeaconClientNetStats->ReceivedControlMessages.IsEmpty());
+		UTEST_TRUE_EXPR(BeaconHostNetStats->ReceivedControlMessages.IsEmpty());
+
+		// Handshake complete, channel is open.
+
+		// Disconnect the client.
+		BeaconClient->DestroyBeacon();
+		BeaconClient = nullptr;
+
+		// Check that client actor on the host cleaned up.
+		UTEST_TRUE_EXPR(BeaconUnitTest::TickUntil(*Prerequisites, [&Prerequisites](){ return Prerequisites->GetStats().HostObject.NotifyClientDisconnected.InvokeCount == 1; }));
+		HostUserBeaconClient = BeaconUnitTest::GetBeaconClientForUser(BeaconHostObject, UserId);
+		UTEST_TRUE_EXPR(HostUserBeaconClient == nullptr);
+	}
+
+	return true;
+}
+
+IMPLEMENT_SIMPLE_AUTOMATION_TEST(FOnlineBeaconTestAuthenticatedHandshakeFailure,
+	"System.Engine.Online.OnlineSubsystemUtils.OnlineBeacon.AuthenticatedHandshakeFailure",
+	EAutomationTestFlags::EditorContext | EAutomationTestFlags::EngineFilter)
+bool FOnlineBeaconTestAuthenticatedHandshakeFailure::RunTest(const FString& Parameters)
+{
+	TSharedPtr<BeaconUnitTest::FTestPrerequisites> Prerequisites = BeaconUnitTest::FTestPrerequisites::TryCreate();
+	UTEST_TRUE_EXPR(Prerequisites.IsValid());
+
+	Prerequisites->GetConfig().Auth.bEnabled = true;
+	Prerequisites->GetConfig().NetDriver.bNoTimeouts = true;
+	Prerequisites->GetConfig().NetDriver.ServerListenPort = 9999;
+
+	// Snapshot config to restore before running a test section.
+	const BeaconUnitTest::FTestConfig BaseConfig = Prerequisites->GetConfig();
+
+	FUniqueNetIdStringRef UserId = FUniqueNetIdString::Create(TEXT("User"), TEXT("UnitTest"));
+
+	// Host setup.
+	TObjectPtr<AOnlineBeaconUnitTestHostObject> BeaconHostObject = Prerequisites->GetWorld()->SpawnActor<AOnlineBeaconUnitTestHostObject>(AOnlineBeaconUnitTestHostObject::StaticClass());
+	TObjectPtr<AOnlineBeaconUnitTestHost> BeaconHost = Prerequisites->GetWorld()->SpawnActor<AOnlineBeaconUnitTestHost>(AOnlineBeaconUnitTestHost::StaticClass());
+	ON_SCOPE_EXIT{ if (BeaconHostObject) { BeaconHostObject->Destroy(); } };
+	ON_SCOPE_EXIT{ if (BeaconHost) { BeaconHost->DestroyBeacon(); } };
+	UTEST_TRUE_EXPR(BeaconUnitTest::ConfigureBeacon(*Prerequisites, BeaconHost));
+	BeaconHost->ListenPort = Prerequisites->GetConfig().NetDriver.ServerListenPort;
+	BeaconHost->RegisterHost(BeaconHostObject);
+	UTEST_TRUE_EXPR(BeaconHost->InitHost());
+	BeaconHost->PauseBeaconRequests(false);
+
+	TSharedPtr<BeaconUnitTest::FNetworkStats> BeaconHostNetStats;
+	UTEST_TRUE_EXPR(BeaconUnitTest::ConfigureBeaconNetDriver(*Prerequisites, BeaconHost, &BeaconHostNetStats));
+	UTEST_TRUE_EXPR(BeaconHostNetStats.IsValid());
+
+	// Method 1.
+	// Host handshake calls StartVerifyAuthentication and expects OnAuthenticationVerificationComplete to be called to complete authentication.
+	// Test OnAuthenticationVerificationComplete fired with no delay is working.
+	{
+		BeaconHostNetStats->ReceivedControlMessages.Empty();
+		Prerequisites->GetConfig() = BaseConfig;
+		Prerequisites->GetStats() = BeaconUnitTest::FTestStats();
+
+		Prerequisites->GetConfig().Auth.Method1.bEnabled = true;
+		Prerequisites->GetConfig().Auth.Method1.Result = OnlineIdentity::Errors::InvalidCreds();
+
+		// Client setup.
+
+		TObjectPtr<AOnlineBeaconUnitTestClient> BeaconClient = Prerequisites->GetWorld()->SpawnActor<AOnlineBeaconUnitTestClient>(AOnlineBeaconUnitTestClient::StaticClass());
+		ON_SCOPE_EXIT{ if (BeaconClient) { BeaconClient->DestroyBeacon(); } };
+		UTEST_TRUE_EXPR(BeaconUnitTest::ConfigureBeacon(*Prerequisites, BeaconClient));
+
+		UTEST_TRUE_EXPR(BeaconClient->GetConnectionState() == EBeaconConnectionState::Invalid);
+		UTEST_TRUE_EXPR(BeaconUnitTest::InitClientForUser(*Prerequisites, BeaconClient, UserId));
+		UTEST_TRUE_EXPR(BeaconClient->GetConnectionState() == EBeaconConnectionState::Pending);
+		UTEST_TRUE_EXPR(BeaconClient->GetNetConnection() != nullptr);
+
+		TSharedPtr<BeaconUnitTest::FNetworkStats> BeaconClientNetStats;
+		UTEST_TRUE_EXPR(BeaconUnitTest::ConfigureBeaconNetDriver(*Prerequisites, BeaconClient, &BeaconClientNetStats));
+		UTEST_TRUE_EXPR(BeaconClientNetStats.IsValid());
+
+		// Handshake testing.
+
+		UTEST_TRUE_EXPR(BeaconUnitTest::TickUntilControlMessageReceived(*Prerequisites, *BeaconHostNetStats));
+		UTEST_TRUE_EXPR(BeaconClient->GetConnectionState() == EBeaconConnectionState::Pending);
+		UTEST_TRUE_EXPR(BeaconClientNetStats->ReceivedControlMessages.IsEmpty());
+		UTEST_TRUE_EXPR(BeaconHostNetStats->ReceivedControlMessages == TArray<uint8>({ NMT_Hello }));
+		UTEST_TRUE_EXPR(BeaconUnitTest::GetBeaconClientForUser(BeaconHostObject, UserId) == nullptr);
+		BeaconHostNetStats->ReceivedControlMessages.Empty();
+
+		UTEST_TRUE_EXPR(BeaconUnitTest::TickOnce(*Prerequisites));
+		UTEST_TRUE_EXPR(BeaconClient->GetConnectionState() == EBeaconConnectionState::Pending);
+		UTEST_TRUE_EXPR(BeaconClientNetStats->ReceivedControlMessages == TArray<uint8>({ NMT_Challenge }));
+		UTEST_TRUE_EXPR(BeaconHostNetStats->ReceivedControlMessages.IsEmpty());
+		BeaconClientNetStats->ReceivedControlMessages.Empty();
+
+		UTEST_TRUE_EXPR(BeaconUnitTest::TickOnce(*Prerequisites));
+		UTEST_TRUE_EXPR(BeaconClient->GetConnectionState() == EBeaconConnectionState::Pending);
+		UTEST_TRUE_EXPR(BeaconClientNetStats->ReceivedControlMessages.IsEmpty());
+		UTEST_TRUE_EXPR(BeaconHostNetStats->ReceivedControlMessages == TArray<uint8>({ NMT_Login }));
+		UTEST_TRUE_EXPR(BeaconUnitTest::GetBeaconClientForUser(BeaconHostObject, UserId) == nullptr);
+		BeaconHostNetStats->ReceivedControlMessages.Empty();
+
+		UTEST_TRUE_EXPR(BeaconUnitTest::TickOnce(*Prerequisites));
+		UTEST_TRUE_EXPR(BeaconClient->GetConnectionState() == EBeaconConnectionState::Invalid);
+		UTEST_TRUE_EXPR(BeaconClientNetStats->ReceivedControlMessages == TArray<uint8>({ NMT_Failure }));
+		UTEST_TRUE_EXPR(BeaconHostNetStats->ReceivedControlMessages.IsEmpty());
+		UTEST_TRUE_EXPR(Prerequisites->GetStats().Client.OnFailure.InvokeCount == 1);
+		BeaconClientNetStats->ReceivedControlMessages.Empty();
+	}
+
+	// Method 1.
+	// Host handshake calls StartVerifyAuthentication and expects OnAuthenticationVerificationComplete to be called to complete authentication.
+	// Test OnAuthenticationVerificationComplete fired with a single frame delay is working.
+	{
+		BeaconHostNetStats->ReceivedControlMessages.Empty();
+		Prerequisites->GetConfig() = BaseConfig;
+		Prerequisites->GetStats() = BeaconUnitTest::FTestStats();
+
+		Prerequisites->GetConfig().Auth.Method1.bEnabled = true;
+		Prerequisites->GetConfig().Auth.Method1.bDelayDelegate = true;
+		Prerequisites->GetConfig().Auth.Method1.Result = OnlineIdentity::Errors::InvalidCreds();
+
+		// Client setup.
+
+		TObjectPtr<AOnlineBeaconUnitTestClient> BeaconClient = Prerequisites->GetWorld()->SpawnActor<AOnlineBeaconUnitTestClient>(AOnlineBeaconUnitTestClient::StaticClass());
+		ON_SCOPE_EXIT{ if (BeaconClient) { BeaconClient->DestroyBeacon(); } };
+		UTEST_TRUE_EXPR(BeaconUnitTest::ConfigureBeacon(*Prerequisites, BeaconClient));
+
+		UTEST_TRUE_EXPR(BeaconClient->GetConnectionState() == EBeaconConnectionState::Invalid);
+		UTEST_TRUE_EXPR(BeaconUnitTest::InitClientForUser(*Prerequisites, BeaconClient, UserId));
+		UTEST_TRUE_EXPR(BeaconClient->GetConnectionState() == EBeaconConnectionState::Pending);
+		UTEST_TRUE_EXPR(BeaconClient->GetNetConnection() != nullptr);
+
+		TSharedPtr<BeaconUnitTest::FNetworkStats> BeaconClientNetStats;
+		UTEST_TRUE_EXPR(BeaconUnitTest::ConfigureBeaconNetDriver(*Prerequisites, BeaconClient, &BeaconClientNetStats));
+		UTEST_TRUE_EXPR(BeaconClientNetStats.IsValid());
+
+		// Handshake testing.
+
+		UTEST_TRUE_EXPR(BeaconUnitTest::TickUntilControlMessageReceived(*Prerequisites, *BeaconHostNetStats));
+		UTEST_TRUE_EXPR(BeaconClient->GetConnectionState() == EBeaconConnectionState::Pending);
+		UTEST_TRUE_EXPR(BeaconClientNetStats->ReceivedControlMessages.IsEmpty());
+		UTEST_TRUE_EXPR(BeaconHostNetStats->ReceivedControlMessages == TArray<uint8>({ NMT_Hello }));
+		UTEST_TRUE_EXPR(BeaconUnitTest::GetBeaconClientForUser(BeaconHostObject, UserId) == nullptr);
+		BeaconHostNetStats->ReceivedControlMessages.Empty();
+
+		UTEST_TRUE_EXPR(BeaconUnitTest::TickOnce(*Prerequisites));
+		UTEST_TRUE_EXPR(BeaconClient->GetConnectionState() == EBeaconConnectionState::Pending);
+		UTEST_TRUE_EXPR(BeaconClientNetStats->ReceivedControlMessages == TArray<uint8>({ NMT_Challenge }));
+		UTEST_TRUE_EXPR(BeaconHostNetStats->ReceivedControlMessages.IsEmpty());
+		BeaconClientNetStats->ReceivedControlMessages.Empty();
+
+		UTEST_TRUE_EXPR(BeaconUnitTest::TickOnce(*Prerequisites));
+		UTEST_TRUE_EXPR(BeaconClient->GetConnectionState() == EBeaconConnectionState::Pending);
+		UTEST_TRUE_EXPR(BeaconClientNetStats->ReceivedControlMessages.IsEmpty());
+		UTEST_TRUE_EXPR(BeaconHostNetStats->ReceivedControlMessages == TArray<uint8>({ NMT_Login }));
+		UTEST_TRUE_EXPR(BeaconUnitTest::GetBeaconClientForUser(BeaconHostObject, UserId) == nullptr);
+		BeaconHostNetStats->ReceivedControlMessages.Empty();
+
+		UTEST_TRUE_EXPR(BeaconUnitTest::TickOnce(*Prerequisites));
+		UTEST_TRUE_EXPR(BeaconClient->GetConnectionState() == EBeaconConnectionState::Pending);
+		UTEST_TRUE_EXPR(BeaconClientNetStats->ReceivedControlMessages.IsEmpty());
+		UTEST_TRUE_EXPR(BeaconHostNetStats->ReceivedControlMessages.IsEmpty());
+
+		UTEST_TRUE_EXPR(BeaconUnitTest::TickOnce(*Prerequisites));
+		UTEST_TRUE_EXPR(BeaconClient->GetConnectionState() == EBeaconConnectionState::Invalid);
+		UTEST_TRUE_EXPR(BeaconClientNetStats->ReceivedControlMessages == TArray<uint8>({ NMT_Failure }));
+		UTEST_TRUE_EXPR(BeaconHostNetStats->ReceivedControlMessages.IsEmpty());
+		UTEST_TRUE_EXPR(Prerequisites->GetStats().Client.OnFailure.InvokeCount == 1);
+		BeaconClientNetStats->ReceivedControlMessages.Empty();
+	}
+	
+	// Method 2.
+	// Host handshake calls StartVerifyAuthentication (callback signature) and expects callback delegate to be called to complete authentication.
+	// Test callback delegate fired with no delay is working.
+	{
+		BeaconHostNetStats->ReceivedControlMessages.Empty();
+		Prerequisites->GetConfig() = BaseConfig;
+		Prerequisites->GetStats() = BeaconUnitTest::FTestStats();
+
+		Prerequisites->GetConfig().Auth.Method2.bEnabled = true;
+		Prerequisites->GetConfig().Auth.Method2.Result = OnlineIdentity::Errors::InvalidCreds();
+
+		// Client setup.
+
+		TObjectPtr<AOnlineBeaconUnitTestClient> BeaconClient = Prerequisites->GetWorld()->SpawnActor<AOnlineBeaconUnitTestClient>(AOnlineBeaconUnitTestClient::StaticClass());
+		ON_SCOPE_EXIT{ if (BeaconClient) { BeaconClient->DestroyBeacon(); } };
+		UTEST_TRUE_EXPR(BeaconUnitTest::ConfigureBeacon(*Prerequisites, BeaconClient));
+
+		UTEST_TRUE_EXPR(BeaconClient->GetConnectionState() == EBeaconConnectionState::Invalid);
+		UTEST_TRUE_EXPR(BeaconUnitTest::InitClientForUser(*Prerequisites, BeaconClient, UserId));
+		UTEST_TRUE_EXPR(BeaconClient->GetConnectionState() == EBeaconConnectionState::Pending);
+		UTEST_TRUE_EXPR(BeaconClient->GetNetConnection() != nullptr);
+
+		TSharedPtr<BeaconUnitTest::FNetworkStats> BeaconClientNetStats;
+		UTEST_TRUE_EXPR(BeaconUnitTest::ConfigureBeaconNetDriver(*Prerequisites, BeaconClient, &BeaconClientNetStats));
+		UTEST_TRUE_EXPR(BeaconClientNetStats.IsValid());
+
+		// Handshake testing.
+
+		UTEST_TRUE_EXPR(BeaconUnitTest::TickUntilControlMessageReceived(*Prerequisites, *BeaconHostNetStats));
+		UTEST_TRUE_EXPR(BeaconClient->GetConnectionState() == EBeaconConnectionState::Pending);
+		UTEST_TRUE_EXPR(BeaconClientNetStats->ReceivedControlMessages.IsEmpty());
+		UTEST_TRUE_EXPR(BeaconHostNetStats->ReceivedControlMessages == TArray<uint8>({ NMT_Hello }));
+		UTEST_TRUE_EXPR(BeaconUnitTest::GetBeaconClientForUser(BeaconHostObject, UserId) == nullptr);
+		BeaconHostNetStats->ReceivedControlMessages.Empty();
+
+		UTEST_TRUE_EXPR(BeaconUnitTest::TickOnce(*Prerequisites));
+		UTEST_TRUE_EXPR(BeaconClient->GetConnectionState() == EBeaconConnectionState::Pending);
+		UTEST_TRUE_EXPR(BeaconClientNetStats->ReceivedControlMessages == TArray<uint8>({ NMT_Challenge }));
+		UTEST_TRUE_EXPR(BeaconHostNetStats->ReceivedControlMessages.IsEmpty());
+		BeaconClientNetStats->ReceivedControlMessages.Empty();
+
+		UTEST_TRUE_EXPR(BeaconUnitTest::TickOnce(*Prerequisites));
+		UTEST_TRUE_EXPR(BeaconClient->GetConnectionState() == EBeaconConnectionState::Pending);
+		UTEST_TRUE_EXPR(BeaconClientNetStats->ReceivedControlMessages.IsEmpty());
+		UTEST_TRUE_EXPR(BeaconHostNetStats->ReceivedControlMessages == TArray<uint8>({ NMT_Login }));
+		UTEST_TRUE_EXPR(BeaconUnitTest::GetBeaconClientForUser(BeaconHostObject, UserId) == nullptr);
+		BeaconHostNetStats->ReceivedControlMessages.Empty();
+
+		UTEST_TRUE_EXPR(BeaconUnitTest::TickOnce(*Prerequisites));
+		UTEST_TRUE_EXPR(BeaconClient->GetConnectionState() == EBeaconConnectionState::Invalid);
+		UTEST_TRUE_EXPR(BeaconClientNetStats->ReceivedControlMessages == TArray<uint8>({ NMT_Failure }));
+		UTEST_TRUE_EXPR(BeaconHostNetStats->ReceivedControlMessages.IsEmpty());
+		UTEST_TRUE_EXPR(Prerequisites->GetStats().Client.OnFailure.InvokeCount == 1);
+		BeaconClientNetStats->ReceivedControlMessages.Empty();
+	}
+
+	// Method 2.
+	// Host handshake calls StartVerifyAuthentication (callback signature) and expects callback delegate to be called to complete authentication.
+	// Test callback delegate fired with a single frame delay is working.
+	{
+		BeaconHostNetStats->ReceivedControlMessages.Empty();
+		Prerequisites->GetConfig() = BaseConfig;
+		Prerequisites->GetStats() = BeaconUnitTest::FTestStats();
+
+		Prerequisites->GetConfig().Auth.Method2.bEnabled = true;
+		Prerequisites->GetConfig().Auth.Method2.bDelayDelegate = true;
+		Prerequisites->GetConfig().Auth.Method2.Result = OnlineIdentity::Errors::InvalidCreds();
+
+		// Client setup.
+
+		TObjectPtr<AOnlineBeaconUnitTestClient> BeaconClient = Prerequisites->GetWorld()->SpawnActor<AOnlineBeaconUnitTestClient>(AOnlineBeaconUnitTestClient::StaticClass());
+		ON_SCOPE_EXIT{ if (BeaconClient) { BeaconClient->DestroyBeacon(); } };
+		UTEST_TRUE_EXPR(BeaconUnitTest::ConfigureBeacon(*Prerequisites, BeaconClient));
+
+		UTEST_TRUE_EXPR(BeaconClient->GetConnectionState() == EBeaconConnectionState::Invalid);
+		UTEST_TRUE_EXPR(BeaconUnitTest::InitClientForUser(*Prerequisites, BeaconClient, UserId));
+		UTEST_TRUE_EXPR(BeaconClient->GetConnectionState() == EBeaconConnectionState::Pending);
+		UTEST_TRUE_EXPR(BeaconClient->GetNetConnection() != nullptr);
+
+		TSharedPtr<BeaconUnitTest::FNetworkStats> BeaconClientNetStats;
+		UTEST_TRUE_EXPR(BeaconUnitTest::ConfigureBeaconNetDriver(*Prerequisites, BeaconClient, &BeaconClientNetStats));
+		UTEST_TRUE_EXPR(BeaconClientNetStats.IsValid());
+
+		// Handshake testing.
+
+		UTEST_TRUE_EXPR(BeaconUnitTest::TickUntilControlMessageReceived(*Prerequisites, *BeaconHostNetStats));
+		UTEST_TRUE_EXPR(BeaconClient->GetConnectionState() == EBeaconConnectionState::Pending);
+		UTEST_TRUE_EXPR(BeaconClientNetStats->ReceivedControlMessages.IsEmpty());
+		UTEST_TRUE_EXPR(BeaconHostNetStats->ReceivedControlMessages == TArray<uint8>({ NMT_Hello }));
+		UTEST_TRUE_EXPR(BeaconUnitTest::GetBeaconClientForUser(BeaconHostObject, UserId) == nullptr);
+		BeaconHostNetStats->ReceivedControlMessages.Empty();
+
+		UTEST_TRUE_EXPR(BeaconUnitTest::TickOnce(*Prerequisites));
+		UTEST_TRUE_EXPR(BeaconClient->GetConnectionState() == EBeaconConnectionState::Pending);
+		UTEST_TRUE_EXPR(BeaconClientNetStats->ReceivedControlMessages == TArray<uint8>({ NMT_Challenge }));
+		UTEST_TRUE_EXPR(BeaconHostNetStats->ReceivedControlMessages.IsEmpty());
+		BeaconClientNetStats->ReceivedControlMessages.Empty();
+
+		UTEST_TRUE_EXPR(BeaconUnitTest::TickOnce(*Prerequisites));
+		UTEST_TRUE_EXPR(BeaconClient->GetConnectionState() == EBeaconConnectionState::Pending);
+		UTEST_TRUE_EXPR(BeaconClientNetStats->ReceivedControlMessages.IsEmpty());
+		UTEST_TRUE_EXPR(BeaconHostNetStats->ReceivedControlMessages == TArray<uint8>({ NMT_Login }));
+		UTEST_TRUE_EXPR(BeaconUnitTest::GetBeaconClientForUser(BeaconHostObject, UserId) == nullptr);
+		BeaconHostNetStats->ReceivedControlMessages.Empty();
+
+		UTEST_TRUE_EXPR(BeaconUnitTest::TickOnce(*Prerequisites));
+		UTEST_TRUE_EXPR(BeaconClient->GetConnectionState() == EBeaconConnectionState::Pending);
+		UTEST_TRUE_EXPR(BeaconClientNetStats->ReceivedControlMessages.IsEmpty());
+		UTEST_TRUE_EXPR(BeaconHostNetStats->ReceivedControlMessages.IsEmpty());
+
+		UTEST_TRUE_EXPR(BeaconUnitTest::TickOnce(*Prerequisites));
+		UTEST_TRUE_EXPR(BeaconClient->GetConnectionState() == EBeaconConnectionState::Invalid);
+		UTEST_TRUE_EXPR(BeaconClientNetStats->ReceivedControlMessages == TArray<uint8>({ NMT_Failure }));
+		UTEST_TRUE_EXPR(BeaconHostNetStats->ReceivedControlMessages.IsEmpty());
+		UTEST_TRUE_EXPR(Prerequisites->GetStats().Client.OnFailure.InvokeCount == 1);
+		BeaconClientNetStats->ReceivedControlMessages.Empty();
+	}
+
+	// Method 3.
+	// Host handshake calls StartVerifyAuthentication (callback signature with LoginOptions) and expects callback delegate to be called to complete authentication.
+	// Test callback delegate fired with no delay is working.
+	{
+		BeaconHostNetStats->ReceivedControlMessages.Empty();
+		Prerequisites->GetConfig() = BaseConfig;
+		Prerequisites->GetStats() = BeaconUnitTest::FTestStats();
+
+		Prerequisites->GetConfig().Auth.Method3.bEnabled = true;
+		Prerequisites->GetConfig().Auth.Method3.Result = OnlineIdentity::Errors::InvalidCreds();
+
+		// Client setup.
+
+		TObjectPtr<AOnlineBeaconUnitTestClient> BeaconClient = Prerequisites->GetWorld()->SpawnActor<AOnlineBeaconUnitTestClient>(AOnlineBeaconUnitTestClient::StaticClass());
+		ON_SCOPE_EXIT{ if (BeaconClient) { BeaconClient->DestroyBeacon(); } };
+		UTEST_TRUE_EXPR(BeaconUnitTest::ConfigureBeacon(*Prerequisites, BeaconClient));
+
+		UTEST_TRUE_EXPR(BeaconClient->GetConnectionState() == EBeaconConnectionState::Invalid);
+		UTEST_TRUE_EXPR(BeaconUnitTest::InitClientForUser(*Prerequisites, BeaconClient, UserId));
+		UTEST_TRUE_EXPR(BeaconClient->GetConnectionState() == EBeaconConnectionState::Pending);
+		UTEST_TRUE_EXPR(BeaconClient->GetNetConnection() != nullptr);
+
+		TSharedPtr<BeaconUnitTest::FNetworkStats> BeaconClientNetStats;
+		UTEST_TRUE_EXPR(BeaconUnitTest::ConfigureBeaconNetDriver(*Prerequisites, BeaconClient, &BeaconClientNetStats));
+		UTEST_TRUE_EXPR(BeaconClientNetStats.IsValid());
+
+		// Handshake testing.
+
+		UTEST_TRUE_EXPR(BeaconUnitTest::TickUntilControlMessageReceived(*Prerequisites, *BeaconHostNetStats));
+		UTEST_TRUE_EXPR(BeaconClient->GetConnectionState() == EBeaconConnectionState::Pending);
+		UTEST_TRUE_EXPR(BeaconClientNetStats->ReceivedControlMessages.IsEmpty());
+		UTEST_TRUE_EXPR(BeaconHostNetStats->ReceivedControlMessages == TArray<uint8>({ NMT_Hello }));
+		UTEST_TRUE_EXPR(BeaconUnitTest::GetBeaconClientForUser(BeaconHostObject, UserId) == nullptr);
+		BeaconHostNetStats->ReceivedControlMessages.Empty();
+
+		UTEST_TRUE_EXPR(BeaconUnitTest::TickOnce(*Prerequisites));
+		UTEST_TRUE_EXPR(BeaconClient->GetConnectionState() == EBeaconConnectionState::Pending);
+		UTEST_TRUE_EXPR(BeaconClientNetStats->ReceivedControlMessages == TArray<uint8>({ NMT_Challenge }));
+		UTEST_TRUE_EXPR(BeaconHostNetStats->ReceivedControlMessages.IsEmpty());
+		BeaconClientNetStats->ReceivedControlMessages.Empty();
+
+		UTEST_TRUE_EXPR(BeaconUnitTest::TickOnce(*Prerequisites));
+		UTEST_TRUE_EXPR(BeaconClient->GetConnectionState() == EBeaconConnectionState::Pending);
+		UTEST_TRUE_EXPR(BeaconClientNetStats->ReceivedControlMessages.IsEmpty());
+		UTEST_TRUE_EXPR(BeaconHostNetStats->ReceivedControlMessages == TArray<uint8>({ NMT_Login }));
+		UTEST_TRUE_EXPR(BeaconUnitTest::GetBeaconClientForUser(BeaconHostObject, UserId) == nullptr);
+		BeaconHostNetStats->ReceivedControlMessages.Empty();
+
+		UTEST_TRUE_EXPR(BeaconUnitTest::TickOnce(*Prerequisites));
+		UTEST_TRUE_EXPR(BeaconClient->GetConnectionState() == EBeaconConnectionState::Invalid);
+		UTEST_TRUE_EXPR(BeaconClientNetStats->ReceivedControlMessages == TArray<uint8>({ NMT_Failure }));
+		UTEST_TRUE_EXPR(BeaconHostNetStats->ReceivedControlMessages.IsEmpty());
+		UTEST_TRUE_EXPR(Prerequisites->GetStats().Client.OnFailure.InvokeCount == 1);
+		BeaconClientNetStats->ReceivedControlMessages.Empty();
+	}
+
+	// Method 3.
+	// Host handshake calls StartVerifyAuthentication (callback signature with LoginOptions) and expects callback delegate to be called to complete authentication.
+	// Test callback delegate fired with a single frame delay is working.
+	{
+		BeaconHostNetStats->ReceivedControlMessages.Empty();
+		Prerequisites->GetConfig() = BaseConfig;
+		Prerequisites->GetStats() = BeaconUnitTest::FTestStats();
+
+		Prerequisites->GetConfig().Auth.Method3.bEnabled = true;
+		Prerequisites->GetConfig().Auth.Method3.bDelayDelegate = true;
+		Prerequisites->GetConfig().Auth.Method3.Result = OnlineIdentity::Errors::InvalidCreds();
+
+		// Client setup.
+
+		TObjectPtr<AOnlineBeaconUnitTestClient> BeaconClient = Prerequisites->GetWorld()->SpawnActor<AOnlineBeaconUnitTestClient>(AOnlineBeaconUnitTestClient::StaticClass());
+		ON_SCOPE_EXIT{ if (BeaconClient) { BeaconClient->DestroyBeacon(); } };
+		UTEST_TRUE_EXPR(BeaconUnitTest::ConfigureBeacon(*Prerequisites, BeaconClient));
+
+		UTEST_TRUE_EXPR(BeaconClient->GetConnectionState() == EBeaconConnectionState::Invalid);
+		UTEST_TRUE_EXPR(BeaconUnitTest::InitClientForUser(*Prerequisites, BeaconClient, UserId));
+		UTEST_TRUE_EXPR(BeaconClient->GetConnectionState() == EBeaconConnectionState::Pending);
+		UTEST_TRUE_EXPR(BeaconClient->GetNetConnection() != nullptr);
+
+		TSharedPtr<BeaconUnitTest::FNetworkStats> BeaconClientNetStats;
+		UTEST_TRUE_EXPR(BeaconUnitTest::ConfigureBeaconNetDriver(*Prerequisites, BeaconClient, &BeaconClientNetStats));
+		UTEST_TRUE_EXPR(BeaconClientNetStats.IsValid());
+
+		// Handshake testing.
+
+		UTEST_TRUE_EXPR(BeaconUnitTest::TickUntilControlMessageReceived(*Prerequisites, *BeaconHostNetStats));
+		UTEST_TRUE_EXPR(BeaconClient->GetConnectionState() == EBeaconConnectionState::Pending);
+		UTEST_TRUE_EXPR(BeaconClientNetStats->ReceivedControlMessages.IsEmpty());
+		UTEST_TRUE_EXPR(BeaconHostNetStats->ReceivedControlMessages == TArray<uint8>({ NMT_Hello }));
+		UTEST_TRUE_EXPR(BeaconUnitTest::GetBeaconClientForUser(BeaconHostObject, UserId) == nullptr);
+		BeaconHostNetStats->ReceivedControlMessages.Empty();
+
+		UTEST_TRUE_EXPR(BeaconUnitTest::TickOnce(*Prerequisites));
+		UTEST_TRUE_EXPR(BeaconClient->GetConnectionState() == EBeaconConnectionState::Pending);
+		UTEST_TRUE_EXPR(BeaconClientNetStats->ReceivedControlMessages == TArray<uint8>({ NMT_Challenge }));
+		UTEST_TRUE_EXPR(BeaconHostNetStats->ReceivedControlMessages.IsEmpty());
+		BeaconClientNetStats->ReceivedControlMessages.Empty();
+
+		UTEST_TRUE_EXPR(BeaconUnitTest::TickOnce(*Prerequisites));
+		UTEST_TRUE_EXPR(BeaconClient->GetConnectionState() == EBeaconConnectionState::Pending);
+		UTEST_TRUE_EXPR(BeaconClientNetStats->ReceivedControlMessages.IsEmpty());
+		UTEST_TRUE_EXPR(BeaconHostNetStats->ReceivedControlMessages == TArray<uint8>({ NMT_Login }));
+		UTEST_TRUE_EXPR(BeaconUnitTest::GetBeaconClientForUser(BeaconHostObject, UserId) == nullptr);
+		BeaconHostNetStats->ReceivedControlMessages.Empty();
+
+		UTEST_TRUE_EXPR(BeaconUnitTest::TickOnce(*Prerequisites));
+		UTEST_TRUE_EXPR(BeaconClient->GetConnectionState() == EBeaconConnectionState::Pending);
+		UTEST_TRUE_EXPR(BeaconClientNetStats->ReceivedControlMessages.IsEmpty());
+		UTEST_TRUE_EXPR(BeaconHostNetStats->ReceivedControlMessages.IsEmpty());
+
+		UTEST_TRUE_EXPR(BeaconUnitTest::TickOnce(*Prerequisites));
+		UTEST_TRUE_EXPR(BeaconClient->GetConnectionState() == EBeaconConnectionState::Invalid);
+		UTEST_TRUE_EXPR(BeaconClientNetStats->ReceivedControlMessages == TArray<uint8>({ NMT_Failure }));
+		UTEST_TRUE_EXPR(BeaconHostNetStats->ReceivedControlMessages.IsEmpty());
+		UTEST_TRUE_EXPR(Prerequisites->GetStats().Client.OnFailure.InvokeCount == 1);
+		BeaconClientNetStats->ReceivedControlMessages.Empty();
+	}
+
+	// Auth verify.
+	// Use Method 3 auth and verify that a user supplied VerifyJoinForBeaconType method does not allow the connection to proceed.
+	{
+		BeaconHostNetStats->ReceivedControlMessages.Empty();
+		Prerequisites->GetConfig() = BaseConfig;
+		Prerequisites->GetStats() = BeaconUnitTest::FTestStats();
+
+		Prerequisites->GetConfig().Auth.Method3.bEnabled = true;
+		Prerequisites->GetConfig().Auth.Method3.Result = FOnlineError::Success();
+
+		Prerequisites->GetConfig().Auth.Verify.bEnabled = true;
+		Prerequisites->GetConfig().Auth.Verify.bResult = false;
+
+		// Client setup.
+
+		TObjectPtr<AOnlineBeaconUnitTestClient> BeaconClient = Prerequisites->GetWorld()->SpawnActor<AOnlineBeaconUnitTestClient>(AOnlineBeaconUnitTestClient::StaticClass());
+		ON_SCOPE_EXIT{ if (BeaconClient) { BeaconClient->DestroyBeacon(); } };
+		UTEST_TRUE_EXPR(BeaconUnitTest::ConfigureBeacon(*Prerequisites, BeaconClient));
+
+		UTEST_TRUE_EXPR(BeaconClient->GetConnectionState() == EBeaconConnectionState::Invalid);
+		UTEST_TRUE_EXPR(BeaconUnitTest::InitClientForUser(*Prerequisites, BeaconClient, UserId));
+		UTEST_TRUE_EXPR(BeaconClient->GetConnectionState() == EBeaconConnectionState::Pending);
+		UTEST_TRUE_EXPR(BeaconClient->GetNetConnection() != nullptr);
+
+		TSharedPtr<BeaconUnitTest::FNetworkStats> BeaconClientNetStats;
+		UTEST_TRUE_EXPR(BeaconUnitTest::ConfigureBeaconNetDriver(*Prerequisites, BeaconClient, &BeaconClientNetStats));
+		UTEST_TRUE_EXPR(BeaconClientNetStats.IsValid());
+
+		// Handshake testing.
+
+		UTEST_TRUE_EXPR(BeaconUnitTest::TickUntilControlMessageReceived(*Prerequisites, *BeaconHostNetStats));
+		UTEST_TRUE_EXPR(BeaconClient->GetConnectionState() == EBeaconConnectionState::Pending);
+		UTEST_TRUE_EXPR(BeaconClientNetStats->ReceivedControlMessages.IsEmpty());
+		UTEST_TRUE_EXPR(BeaconHostNetStats->ReceivedControlMessages == TArray<uint8>({ NMT_Hello }));
+		UTEST_TRUE_EXPR(BeaconUnitTest::GetBeaconClientForUser(BeaconHostObject, UserId) == nullptr);
+		BeaconHostNetStats->ReceivedControlMessages.Empty();
+
+		UTEST_TRUE_EXPR(BeaconUnitTest::TickOnce(*Prerequisites));
+		UTEST_TRUE_EXPR(BeaconClient->GetConnectionState() == EBeaconConnectionState::Pending);
+		UTEST_TRUE_EXPR(BeaconClientNetStats->ReceivedControlMessages == TArray<uint8>({ NMT_Challenge }));
+		UTEST_TRUE_EXPR(BeaconHostNetStats->ReceivedControlMessages.IsEmpty());
+		BeaconClientNetStats->ReceivedControlMessages.Empty();
+
+		UTEST_TRUE_EXPR(BeaconUnitTest::TickOnce(*Prerequisites));
+		UTEST_TRUE_EXPR(BeaconClient->GetConnectionState() == EBeaconConnectionState::Pending);
+		UTEST_TRUE_EXPR(BeaconClientNetStats->ReceivedControlMessages.IsEmpty());
+		UTEST_TRUE_EXPR(BeaconHostNetStats->ReceivedControlMessages == TArray<uint8>({ NMT_Login }));
+		UTEST_TRUE_EXPR(BeaconUnitTest::GetBeaconClientForUser(BeaconHostObject, UserId) == nullptr);
+		BeaconHostNetStats->ReceivedControlMessages.Empty();
+
+		UTEST_TRUE_EXPR(BeaconUnitTest::TickOnce(*Prerequisites));
+		UTEST_TRUE_EXPR(BeaconClient->GetConnectionState() == EBeaconConnectionState::Pending);
+		UTEST_TRUE_EXPR(BeaconClientNetStats->ReceivedControlMessages == TArray<uint8>({ NMT_BeaconWelcome }));
+		UTEST_TRUE_EXPR(BeaconHostNetStats->ReceivedControlMessages.IsEmpty());
+		BeaconClientNetStats->ReceivedControlMessages.Empty();
+
+		UTEST_TRUE_EXPR(BeaconUnitTest::TickOnce(*Prerequisites));
+		UTEST_TRUE_EXPR(BeaconClient->GetConnectionState() == EBeaconConnectionState::Pending);
+		UTEST_TRUE_EXPR(BeaconClientNetStats->ReceivedControlMessages.IsEmpty());
+		UTEST_TRUE_EXPR(BeaconHostNetStats->ReceivedControlMessages == TArray<uint8>({ NMT_Netspeed, NMT_BeaconJoin }));
+		UTEST_TRUE_EXPR(BeaconUnitTest::GetBeaconClientForUser(BeaconHostObject, UserId) == nullptr);
+		BeaconHostNetStats->ReceivedControlMessages.Empty();
+
+		UTEST_TRUE_EXPR(BeaconUnitTest::TickOnce(*Prerequisites));
+		UTEST_TRUE_EXPR(BeaconClient->GetConnectionState() == EBeaconConnectionState::Invalid);
+		UTEST_TRUE_EXPR(BeaconClientNetStats->ReceivedControlMessages == TArray<uint8>({ NMT_Failure }));
+		UTEST_TRUE_EXPR(BeaconHostNetStats->ReceivedControlMessages.IsEmpty());
+		UTEST_TRUE_EXPR(Prerequisites->GetStats().Client.OnFailure.InvokeCount == 1);
+		BeaconClientNetStats->ReceivedControlMessages.Empty();
+	}
+	
+	return true;
+}
+
+IMPLEMENT_SIMPLE_AUTOMATION_TEST(FOnlineBeaconTestAuthenticatedHandshakeClientTimeout,
+	"System.Engine.Online.OnlineSubsystemUtils.OnlineBeacon.AuthenticatedHandshakeClientTimeout",
+	EAutomationTestFlags::EditorContext | EAutomationTestFlags::EngineFilter)
+bool FOnlineBeaconTestAuthenticatedHandshakeClientTimeout::RunTest(const FString& Parameters)
+{
+	TSharedPtr<BeaconUnitTest::FTestPrerequisites> Prerequisites = BeaconUnitTest::FTestPrerequisites::TryCreate();
+	UTEST_TRUE_EXPR(Prerequisites.IsValid());
+
+	Prerequisites->GetConfig().Auth.bEnabled = true;
+	Prerequisites->GetConfig().NetDriver.bNoTimeouts = true;
+	Prerequisites->GetConfig().NetDriver.ServerListenPort = 9999;
+
+	// Snapshot config to restore before running a test section.
+	const BeaconUnitTest::FTestConfig BaseConfig = Prerequisites->GetConfig();
+
+	FUniqueNetIdStringRef UserId = FUniqueNetIdString::Create(TEXT("User"), TEXT("UnitTest"));
+
+	// Host setup.
+	TObjectPtr<AOnlineBeaconUnitTestHostObject> BeaconHostObject = Prerequisites->GetWorld()->SpawnActor<AOnlineBeaconUnitTestHostObject>(AOnlineBeaconUnitTestHostObject::StaticClass());
+	TObjectPtr<AOnlineBeaconUnitTestHost> BeaconHost = Prerequisites->GetWorld()->SpawnActor<AOnlineBeaconUnitTestHost>(AOnlineBeaconUnitTestHost::StaticClass());
+	ON_SCOPE_EXIT{ if (BeaconHostObject) { BeaconHostObject->Destroy(); } };
+	ON_SCOPE_EXIT{ if (BeaconHost) { BeaconHost->DestroyBeacon(); } };
+	UTEST_TRUE_EXPR(BeaconUnitTest::ConfigureBeacon(*Prerequisites, BeaconHost));
+	BeaconHost->ListenPort = Prerequisites->GetConfig().NetDriver.ServerListenPort;
+	BeaconHost->RegisterHost(BeaconHostObject);
+	UTEST_TRUE_EXPR(BeaconHost->InitHost());
+	BeaconHost->PauseBeaconRequests(false);
+
+	TSharedPtr<BeaconUnitTest::FNetworkStats> BeaconHostNetStats;
+	UTEST_TRUE_EXPR(BeaconUnitTest::ConfigureBeaconNetDriver(*Prerequisites, BeaconHost, &BeaconHostNetStats));
+	UTEST_TRUE_EXPR(BeaconHostNetStats.IsValid());
+
+	// Timeout after client sends NMT_Login
+	{
+		BeaconHostNetStats->ReceivedControlMessages.Empty();
+		Prerequisites->GetConfig() = BaseConfig;
+		Prerequisites->GetStats() = BeaconUnitTest::FTestStats();
+
+		Prerequisites->GetConfig().NetDriver.bNoTimeouts = false;
+
+		// Client setup.
+
+		TObjectPtr<AOnlineBeaconUnitTestClient> BeaconClient = Prerequisites->GetWorld()->SpawnActor<AOnlineBeaconUnitTestClient>(AOnlineBeaconUnitTestClient::StaticClass());
+		ON_SCOPE_EXIT{ if (BeaconClient) { BeaconClient->DestroyBeacon(); } };
+		UTEST_TRUE_EXPR(BeaconUnitTest::ConfigureBeacon(*Prerequisites, BeaconClient));
+
+		UTEST_TRUE_EXPR(BeaconUnitTest::InitClientForUser(*Prerequisites, BeaconClient, UserId));
+		UTEST_TRUE_EXPR(BeaconClient->GetConnectionState() == EBeaconConnectionState::Pending);
+		UTEST_TRUE_EXPR(BeaconClient->GetNetConnection() != nullptr);
+
+		TSharedPtr<BeaconUnitTest::FNetworkStats> BeaconClientNetStats;
+		UTEST_TRUE_EXPR(BeaconUnitTest::ConfigureBeaconNetDriver(*Prerequisites, BeaconClient, &BeaconClientNetStats));
+		UTEST_TRUE_EXPR(BeaconClientNetStats.IsValid());
+
+		// Handshake testing.
+
+		UTEST_TRUE_EXPR(BeaconUnitTest::TickUntilControlMessageReceived(*Prerequisites, *BeaconHostNetStats));
+		UTEST_TRUE_EXPR(BeaconClient->GetConnectionState() == EBeaconConnectionState::Pending);
+		UTEST_TRUE_EXPR(BeaconClientNetStats->ReceivedControlMessages.IsEmpty());
+		UTEST_TRUE_EXPR(BeaconHostNetStats->ReceivedControlMessages == TArray<uint8>({ NMT_Hello }));
+		BeaconHostNetStats->ReceivedControlMessages.Empty();
+
+		// Block the client from sending packets.
+		UTEST_TRUE_EXPR(BeaconUnitTest::SetSocketFlags(BeaconClient, BeaconUnitTest::ESocketFlags::RecvEnabled));
+
+		UTEST_TRUE_EXPR(BeaconUnitTest::TickUntilDisconnected(*Prerequisites, BeaconClient, BeaconUnitTest::ETickFlags::SleepTickTime));
+		UTEST_TRUE_EXPR(BeaconClient->GetConnectionState() == EBeaconConnectionState::Invalid);
+		UTEST_TRUE_EXPR(BeaconClientNetStats->ReceivedControlMessages == TArray<uint8>({ NMT_Challenge }));
+		UTEST_TRUE_EXPR(BeaconHostNetStats->ReceivedControlMessages.IsEmpty());
+	}
+	
+	return true;
+}
+
+IMPLEMENT_SIMPLE_AUTOMATION_TEST(FOnlineBeaconTestAuthenticatedHandshakeHostTimeout,
+	"System.Engine.Online.OnlineSubsystemUtils.OnlineBeacon.AuthenticatedHandshakeHostTimeout",
+	EAutomationTestFlags::EditorContext | EAutomationTestFlags::EngineFilter)
+bool FOnlineBeaconTestAuthenticatedHandshakeHostTimeout::RunTest(const FString& Parameters)
+{
+	TSharedPtr<BeaconUnitTest::FTestPrerequisites> Prerequisites = BeaconUnitTest::FTestPrerequisites::TryCreate();
+	UTEST_TRUE_EXPR(Prerequisites.IsValid());
+
+	Prerequisites->GetConfig().Auth.bEnabled = true;
+	Prerequisites->GetConfig().NetDriver.ServerListenPort = 9999;
+
+	// Snapshot config to restore before running a test section.
+	const BeaconUnitTest::FTestConfig BaseConfig = Prerequisites->GetConfig();
+
+	FUniqueNetIdStringRef UserId = FUniqueNetIdString::Create(TEXT("User"), TEXT("UnitTest"));
+
+	// Host setup.
+	TObjectPtr<AOnlineBeaconUnitTestHostObject> BeaconHostObject = Prerequisites->GetWorld()->SpawnActor<AOnlineBeaconUnitTestHostObject>(AOnlineBeaconUnitTestHostObject::StaticClass());
+	TObjectPtr<AOnlineBeaconUnitTestHost> BeaconHost = Prerequisites->GetWorld()->SpawnActor<AOnlineBeaconUnitTestHost>(AOnlineBeaconUnitTestHost::StaticClass());
+	ON_SCOPE_EXIT{ if (BeaconHostObject) { BeaconHostObject->Destroy(); } };
+	ON_SCOPE_EXIT{ if (BeaconHost) { BeaconHost->DestroyBeacon(); } };
+	UTEST_TRUE_EXPR(BeaconUnitTest::ConfigureBeacon(*Prerequisites, BeaconHost));
+	BeaconHost->ListenPort = Prerequisites->GetConfig().NetDriver.ServerListenPort;
+	BeaconHost->RegisterHost(BeaconHostObject);
+	UTEST_TRUE_EXPR(BeaconHost->InitHost());
+	BeaconHost->PauseBeaconRequests(false);
+
+	TSharedPtr<BeaconUnitTest::FNetworkStats> BeaconHostNetStats;
+	UTEST_TRUE_EXPR(BeaconUnitTest::ConfigureBeaconNetDriver(*Prerequisites, BeaconHost, &BeaconHostNetStats));
+	UTEST_TRUE_EXPR(BeaconHostNetStats.IsValid());
+
+	// Timeout after host sends NMT_Challenge
+	{
+		BeaconHostNetStats->ReceivedControlMessages.Empty();
+		Prerequisites->GetConfig() = BaseConfig;
+		Prerequisites->GetStats() = BeaconUnitTest::FTestStats();
+
+		Prerequisites->GetConfig().NetDriver.bNoTimeouts = true;
+
+		// Client setup.
+
+		TObjectPtr<AOnlineBeaconUnitTestClient> BeaconClient = Prerequisites->GetWorld()->SpawnActor<AOnlineBeaconUnitTestClient>(AOnlineBeaconUnitTestClient::StaticClass());
+		ON_SCOPE_EXIT{ if (BeaconClient) { BeaconClient->DestroyBeacon(); } };
+		UTEST_TRUE_EXPR(BeaconUnitTest::ConfigureBeacon(*Prerequisites, BeaconClient));
+
+		UTEST_TRUE_EXPR(BeaconUnitTest::InitClientForUser(*Prerequisites, BeaconClient, UserId));
+		UTEST_TRUE_EXPR(BeaconClient->GetConnectionState() == EBeaconConnectionState::Pending);
+		UTEST_TRUE_EXPR(BeaconClient->GetNetConnection() != nullptr);
+
+		TSharedPtr<BeaconUnitTest::FNetworkStats> BeaconClientNetStats;
+		UTEST_TRUE_EXPR(BeaconUnitTest::ConfigureBeaconNetDriver(*Prerequisites, BeaconClient, &BeaconClientNetStats));
+		UTEST_TRUE_EXPR(BeaconClientNetStats.IsValid());
+
+		// Handshake testing.
+
+		UTEST_TRUE_EXPR(BeaconUnitTest::TickUntilControlMessageReceived(*Prerequisites, *BeaconHostNetStats));
+		UTEST_TRUE_EXPR(BeaconClient->GetConnectionState() == EBeaconConnectionState::Pending);
+		UTEST_TRUE_EXPR(BeaconClientNetStats->ReceivedControlMessages.IsEmpty());
+		UTEST_TRUE_EXPR(BeaconHostNetStats->ReceivedControlMessages == TArray<uint8>({ NMT_Hello }));
+		BeaconHostNetStats->ReceivedControlMessages.Empty();
+
+		// Block the client from handling packets.
+		UTEST_TRUE_EXPR(BeaconUnitTest::SetSocketFlags(BeaconClient, BeaconUnitTest::ESocketFlags::Disabled));
+
+		UTEST_TRUE_EXPR(BeaconUnitTest::TickUntilTimeoutElapsed(*Prerequisites, BeaconUnitTest::ETickFlags::SleepTickTime));
+
+		// Unblock the client from handling packets.
+		UTEST_TRUE_EXPR(BeaconUnitTest::SetSocketFlags(BeaconClient, BeaconUnitTest::ESocketFlags::Default));
+
+		// Try to continue the handshake after the host has cleaned up the client state.
+		// The host will not see the clients control message since it has closed the connection.
+		// The client will be in an invalid state due to receiving the close packet.
+
+		UTEST_TRUE_EXPR(BeaconUnitTest::TickOnce(*Prerequisites));
+		UTEST_TRUE_EXPR(BeaconClient->GetConnectionState() == EBeaconConnectionState::Invalid);
+		UTEST_TRUE_EXPR(BeaconClientNetStats->ReceivedControlMessages == TArray<uint8>({ NMT_Challenge }));
+		UTEST_TRUE_EXPR(BeaconHostNetStats->ReceivedControlMessages.IsEmpty());
+	}
+	
+	return true;
+}
+
+IMPLEMENT_SIMPLE_AUTOMATION_TEST(FOnlineBeaconTestGarbageCollection,
+	"System.Engine.Online.OnlineSubsystemUtils.OnlineBeacon.GarbageCollection",
+	EAutomationTestFlags::EditorContext | EAutomationTestFlags::EngineFilter)
+bool FOnlineBeaconTestGarbageCollection::RunTest(const FString& Parameters)
+{
+	TSharedPtr<BeaconUnitTest::FTestPrerequisites> Prerequisites = BeaconUnitTest::FTestPrerequisites::TryCreate();
+	UTEST_TRUE_EXPR(Prerequisites.IsValid());
+
+	Prerequisites->GetConfig().NetDriver.ServerListenPort = 9999;
+
+	// Don't timeout connections due to a slow GC.
+	Prerequisites->GetConfig().NetDriver.bNoTimeouts = true;
+
+	// Snapshot config to restore before running a test section.
+	const BeaconUnitTest::FTestConfig BaseConfig = Prerequisites->GetConfig();
+
+	FUniqueNetIdStringRef UserId = FUniqueNetIdString::Create(TEXT("User"), TEXT("UnitTest"));
+
+	// Successful handshake.
+	// Garbage collect host beacon.
+	{
+		Prerequisites->GetConfig() = BaseConfig;
+		Prerequisites->GetStats() = BeaconUnitTest::FTestStats();
+
+		// Host setup.
+		TObjectPtr<AOnlineBeaconClient> HostUserBeaconClient;
+		TObjectPtr<AOnlineBeaconUnitTestHostObject> BeaconHostObject = Prerequisites->GetWorld()->SpawnActor<AOnlineBeaconUnitTestHostObject>(AOnlineBeaconUnitTestHostObject::StaticClass());
+		TObjectPtr<AOnlineBeaconUnitTestHost> BeaconHost = Prerequisites->GetWorld()->SpawnActor<AOnlineBeaconUnitTestHost>(AOnlineBeaconUnitTestHost::StaticClass());
+		ON_SCOPE_EXIT{ if (BeaconHostObject) { BeaconHostObject->Destroy(); } };
+		ON_SCOPE_EXIT{ if (BeaconHost) BeaconHost->DestroyBeacon(); };
+		UTEST_TRUE_EXPR(BeaconUnitTest::ConfigureBeacon(*Prerequisites, BeaconHost));
+		BeaconHost->ListenPort = Prerequisites->GetConfig().NetDriver.ServerListenPort;
+		BeaconHost->RegisterHost(BeaconHostObject);
+		UTEST_TRUE_EXPR(BeaconHost->InitHost());
+		BeaconHost->PauseBeaconRequests(false);
+
+		TSharedPtr<BeaconUnitTest::FNetworkStats> BeaconHostNetStats;
+		UTEST_TRUE_EXPR(BeaconUnitTest::ConfigureBeaconNetDriver(*Prerequisites, BeaconHost, &BeaconHostNetStats));
+		UTEST_TRUE_EXPR(BeaconHostNetStats.IsValid());
+
+		// Client setup.
+
+		TObjectPtr<AOnlineBeaconUnitTestClient> BeaconClient = Prerequisites->GetWorld()->SpawnActor<AOnlineBeaconUnitTestClient>(AOnlineBeaconUnitTestClient::StaticClass());
+		ON_SCOPE_EXIT{ if (BeaconClient) { BeaconClient->DestroyBeacon(); } };
+		UTEST_TRUE_EXPR(BeaconUnitTest::ConfigureBeacon(*Prerequisites, BeaconClient));
+
+		UTEST_TRUE_EXPR(BeaconClient->GetConnectionState() == EBeaconConnectionState::Invalid);
+		UTEST_TRUE_EXPR(BeaconUnitTest::InitClientForUser(*Prerequisites, BeaconClient, UserId));
+		UTEST_TRUE_EXPR(BeaconClient->GetConnectionState() == EBeaconConnectionState::Pending);
+		UTEST_TRUE_EXPR(BeaconClient->GetNetConnection() != nullptr);
+
+		TSharedPtr<BeaconUnitTest::FNetworkStats> BeaconClientNetStats;
+		UTEST_TRUE_EXPR(BeaconUnitTest::ConfigureBeaconNetDriver(*Prerequisites, BeaconClient, &BeaconClientNetStats));
+		UTEST_TRUE_EXPR(BeaconClientNetStats.IsValid());
+
+		// Handshake testing.
+
+		UTEST_TRUE_EXPR(BeaconUnitTest::TickUntilConnected(*Prerequisites, BeaconClient));
+		UTEST_TRUE_EXPR(BeaconClient->GetConnectionState() == EBeaconConnectionState::Open);
+		UTEST_TRUE_EXPR(Prerequisites->GetStats().Client.OnConnected.InvokeCount == 1);
+		UTEST_TRUE_EXPR(Prerequisites->GetStats().Client.OnFailure.InvokeCount == 0);
+		UTEST_TRUE_EXPR(Prerequisites->GetStats().Host.OnFailure.InvokeCount == 0);
+		UTEST_TRUE_EXPR(Prerequisites->GetStats().HostObject.OnClientConnected.InvokeCount == 1);
+		UTEST_TRUE_EXPR(Prerequisites->GetStats().HostObject.NotifyClientDisconnected.InvokeCount == 0);
+		HostUserBeaconClient = BeaconUnitTest::GetBeaconClientForUser(BeaconHostObject, UserId);
+		UTEST_TRUE_EXPR(HostUserBeaconClient != nullptr);
+		BeaconClientNetStats->ReceivedControlMessages.Empty();
+		BeaconHostNetStats->ReceivedControlMessages.Empty();
+
+		// Handshake complete, channel is open.
+
+		// Garbage collect the host.
+		BeaconHost->DestroyBeacon();
+		BeaconHost = nullptr;
+		BeaconHostObject = nullptr;
+		GEngine->ForceGarbageCollection(true);
+		UTEST_TRUE_EXPR(BeaconClient->GetConnectionState() == EBeaconConnectionState::Open);
+		UTEST_TRUE_EXPR(Prerequisites->GetStats().Client.OnConnected.InvokeCount == 1);
+		UTEST_TRUE_EXPR(Prerequisites->GetStats().Client.OnFailure.InvokeCount == 0);
+		UTEST_TRUE_EXPR(Prerequisites->GetStats().Host.OnFailure.InvokeCount == 0);
+		UTEST_TRUE_EXPR(Prerequisites->GetStats().HostObject.OnClientConnected.InvokeCount == 1);
+		UTEST_TRUE_EXPR(Prerequisites->GetStats().HostObject.NotifyClientDisconnected.InvokeCount == 1);
+
+		// Check that client and host client object cleaned up.
+		// Make sure host client object cleaned up.
+		UTEST_TRUE_EXPR(BeaconUnitTest::TickOnce(*Prerequisites));
+		UTEST_TRUE_EXPR(BeaconClient->GetConnectionState() == EBeaconConnectionState::Open);
+		UTEST_TRUE_EXPR(Prerequisites->GetStats().Client.OnConnected.InvokeCount == 1);
+		UTEST_TRUE_EXPR(Prerequisites->GetStats().Client.OnFailure.InvokeCount == 0);
+		UTEST_TRUE_EXPR(Prerequisites->GetStats().Host.OnFailure.InvokeCount == 0);
+		UTEST_TRUE_EXPR(Prerequisites->GetStats().HostObject.OnClientConnected.InvokeCount == 1);
+		UTEST_TRUE_EXPR(Prerequisites->GetStats().HostObject.NotifyClientDisconnected.InvokeCount == 1);
+
+		// Make sure client cleaned up.
+		UTEST_TRUE_EXPR(BeaconUnitTest::TickOnce(*Prerequisites));
+		UTEST_TRUE_EXPR(BeaconClient->GetConnectionState() == EBeaconConnectionState::Invalid);
+		UTEST_TRUE_EXPR(Prerequisites->GetStats().Client.OnConnected.InvokeCount == 1);
+		UTEST_TRUE_EXPR(Prerequisites->GetStats().Client.OnFailure.InvokeCount == 1);
+		UTEST_TRUE_EXPR(Prerequisites->GetStats().Host.OnFailure.InvokeCount == 0);
+		UTEST_TRUE_EXPR(Prerequisites->GetStats().HostObject.OnClientConnected.InvokeCount == 1);
+		UTEST_TRUE_EXPR(Prerequisites->GetStats().HostObject.NotifyClientDisconnected.InvokeCount == 1);
+	}
+
+	// Successful handshake.
+	// Garbage collect client beacon.
+	{
+		Prerequisites->GetConfig() = BaseConfig;
+		Prerequisites->GetStats() = BeaconUnitTest::FTestStats();
+
+		// Host setup.
+		TObjectPtr<AOnlineBeaconClient> HostUserBeaconClient;
+		TObjectPtr<AOnlineBeaconUnitTestHostObject> BeaconHostObject = Prerequisites->GetWorld()->SpawnActor<AOnlineBeaconUnitTestHostObject>(AOnlineBeaconUnitTestHostObject::StaticClass());
+		TObjectPtr<AOnlineBeaconUnitTestHost> BeaconHost = Prerequisites->GetWorld()->SpawnActor<AOnlineBeaconUnitTestHost>(AOnlineBeaconUnitTestHost::StaticClass());
+		ON_SCOPE_EXIT{ if (BeaconHostObject) { BeaconHostObject->Destroy(); } };
+		ON_SCOPE_EXIT{ if (BeaconHost) { BeaconHost->DestroyBeacon(); } };
+		UTEST_TRUE_EXPR(BeaconUnitTest::ConfigureBeacon(*Prerequisites, BeaconHost));
+		BeaconHost->ListenPort = Prerequisites->GetConfig().NetDriver.ServerListenPort;
+		BeaconHost->RegisterHost(BeaconHostObject);
+		UTEST_TRUE_EXPR(BeaconHost->InitHost());
+		BeaconHost->PauseBeaconRequests(false);
+
+		TSharedPtr<BeaconUnitTest::FNetworkStats> BeaconHostNetStats;
+		UTEST_TRUE_EXPR(BeaconUnitTest::ConfigureBeaconNetDriver(*Prerequisites, BeaconHost, &BeaconHostNetStats));
+		UTEST_TRUE_EXPR(BeaconHostNetStats.IsValid());
+
+		// Client setup.
+
+		TObjectPtr<AOnlineBeaconUnitTestClient> BeaconClient = Prerequisites->GetWorld()->SpawnActor<AOnlineBeaconUnitTestClient>(AOnlineBeaconUnitTestClient::StaticClass());
+		ON_SCOPE_EXIT{ if (BeaconClient) BeaconClient->DestroyBeacon(); };
+		UTEST_TRUE_EXPR(BeaconUnitTest::ConfigureBeacon(*Prerequisites, BeaconClient));
+
+		UTEST_TRUE_EXPR(BeaconClient->GetConnectionState() == EBeaconConnectionState::Invalid);
+		UTEST_TRUE_EXPR(BeaconUnitTest::InitClientForUser(*Prerequisites, BeaconClient, UserId));
+		UTEST_TRUE_EXPR(BeaconClient->GetConnectionState() == EBeaconConnectionState::Pending);
+		UTEST_TRUE_EXPR(BeaconClient->GetNetConnection() != nullptr);
+
+		TSharedPtr<BeaconUnitTest::FNetworkStats> BeaconClientNetStats;
+		UTEST_TRUE_EXPR(BeaconUnitTest::ConfigureBeaconNetDriver(*Prerequisites, BeaconClient, &BeaconClientNetStats));
+		UTEST_TRUE_EXPR(BeaconClientNetStats.IsValid());
+
+		// Handshake testing.
+
+		UTEST_TRUE_EXPR(BeaconUnitTest::TickUntilConnected(*Prerequisites, BeaconClient));
+		UTEST_TRUE_EXPR(BeaconClient->GetConnectionState() == EBeaconConnectionState::Open);
+		UTEST_TRUE_EXPR(Prerequisites->GetStats().Client.OnConnected.InvokeCount == 1);
+		UTEST_TRUE_EXPR(Prerequisites->GetStats().Client.OnFailure.InvokeCount == 0);
+		UTEST_TRUE_EXPR(Prerequisites->GetStats().Host.OnFailure.InvokeCount == 0);
+		UTEST_TRUE_EXPR(Prerequisites->GetStats().HostObject.OnClientConnected.InvokeCount == 1);
+		UTEST_TRUE_EXPR(Prerequisites->GetStats().HostObject.NotifyClientDisconnected.InvokeCount == 0);
+		HostUserBeaconClient = BeaconUnitTest::GetBeaconClientForUser(BeaconHostObject, UserId);
+		UTEST_TRUE_EXPR(HostUserBeaconClient != nullptr);
+		BeaconClientNetStats->ReceivedControlMessages.Empty();
+		BeaconHostNetStats->ReceivedControlMessages.Empty();
+
+		// Handshake complete, channel is open.
+
+		// Garbage collect the client.
+		BeaconClient->DestroyBeacon();
+		BeaconClient = nullptr;
+		UTEST_TRUE_EXPR(GEngine != nullptr);
+		GEngine->ForceGarbageCollection(true);
+
+		// Tick for client connection to send close.
+		UTEST_TRUE_EXPR(BeaconUnitTest::TickOnce(*Prerequisites));
+		// Tick for host connection to see client close.
+		UTEST_TRUE_EXPR(BeaconUnitTest::TickOnce(*Prerequisites));
+		// Tick for host object to be notified of closure.
+		UTEST_TRUE_EXPR(BeaconUnitTest::TickOnce(*Prerequisites));
+		UTEST_TRUE_EXPR(Prerequisites->GetStats().Client.OnConnected.InvokeCount == 1);
+		UTEST_TRUE_EXPR(Prerequisites->GetStats().Client.OnFailure.InvokeCount == 0);
+		UTEST_TRUE_EXPR(Prerequisites->GetStats().Host.OnFailure.InvokeCount == 0);
+		UTEST_TRUE_EXPR(Prerequisites->GetStats().HostObject.OnClientConnected.InvokeCount == 1);
+		UTEST_TRUE_EXPR(Prerequisites->GetStats().HostObject.NotifyClientDisconnected.InvokeCount == 1);
+		UTEST_TRUE_EXPR(BeaconUnitTest::GetBeaconClientForUser(BeaconHostObject, UserId) == nullptr);
+	}
+
+	// Garbage collect host beacon during handshake.
+	{
+		Prerequisites->GetConfig() = BaseConfig;
+		Prerequisites->GetStats() = BeaconUnitTest::FTestStats();
+
+		// Host setup.
+		TObjectPtr<AOnlineBeaconUnitTestHostObject> BeaconHostObject = Prerequisites->GetWorld()->SpawnActor<AOnlineBeaconUnitTestHostObject>(AOnlineBeaconUnitTestHostObject::StaticClass());
+		TObjectPtr<AOnlineBeaconUnitTestHost> BeaconHost = Prerequisites->GetWorld()->SpawnActor<AOnlineBeaconUnitTestHost>(AOnlineBeaconUnitTestHost::StaticClass());
+		ON_SCOPE_EXIT{ if (BeaconHostObject) { BeaconHostObject->Destroy(); } };
+		ON_SCOPE_EXIT{ if (BeaconHost) { BeaconHost->DestroyBeacon(); } };
+		UTEST_TRUE_EXPR(BeaconUnitTest::ConfigureBeacon(*Prerequisites, BeaconHost));
+		BeaconHost->ListenPort = Prerequisites->GetConfig().NetDriver.ServerListenPort;
+		BeaconHost->RegisterHost(BeaconHostObject);
+		UTEST_TRUE_EXPR(BeaconHost->InitHost());
+		BeaconHost->PauseBeaconRequests(false);
+
+		TSharedPtr<BeaconUnitTest::FNetworkStats> BeaconHostNetStats;
+		UTEST_TRUE_EXPR(BeaconUnitTest::ConfigureBeaconNetDriver(*Prerequisites, BeaconHost, &BeaconHostNetStats));
+		UTEST_TRUE_EXPR(BeaconHostNetStats.IsValid());
+
+		// Client setup.
+
+		TObjectPtr<AOnlineBeaconUnitTestClient> BeaconClient = Prerequisites->GetWorld()->SpawnActor<AOnlineBeaconUnitTestClient>(AOnlineBeaconUnitTestClient::StaticClass());
+		ON_SCOPE_EXIT{ if (BeaconClient) BeaconClient->DestroyBeacon(); };
+		UTEST_TRUE_EXPR(BeaconUnitTest::ConfigureBeacon(*Prerequisites, BeaconClient));
+
+		UTEST_TRUE_EXPR(BeaconClient->GetConnectionState() == EBeaconConnectionState::Invalid);
+		UTEST_TRUE_EXPR(BeaconUnitTest::InitClientForUser(*Prerequisites, BeaconClient, UserId));
+		UTEST_TRUE_EXPR(BeaconClient->GetConnectionState() == EBeaconConnectionState::Pending);
+		UTEST_TRUE_EXPR(BeaconClient->GetNetConnection() != nullptr);
+
+		TSharedPtr<BeaconUnitTest::FNetworkStats> BeaconClientNetStats;
+		UTEST_TRUE_EXPR(BeaconUnitTest::ConfigureBeaconNetDriver(*Prerequisites, BeaconClient, &BeaconClientNetStats));
+		UTEST_TRUE_EXPR(BeaconClientNetStats.IsValid());
+
+		// Handshake testing.
+
+		// Wait until the client receives a control message from the host beacon.
+		UTEST_TRUE_EXPR(BeaconUnitTest::TickUntilControlMessageReceived(*Prerequisites, *BeaconClientNetStats));
+		UTEST_TRUE_EXPR(BeaconClient->GetConnectionState() == EBeaconConnectionState::Pending);
+		UTEST_TRUE_EXPR(BeaconClientNetStats->ReceivedControlMessages == TArray<uint8>({ NMT_BeaconWelcome }));
+		UTEST_TRUE_EXPR(Prerequisites->GetStats().Client.OnConnected.InvokeCount == 0);
+		UTEST_TRUE_EXPR(Prerequisites->GetStats().Client.OnFailure.InvokeCount == 0);
+		UTEST_TRUE_EXPR(Prerequisites->GetStats().Host.OnFailure.InvokeCount == 0);
+		UTEST_TRUE_EXPR(Prerequisites->GetStats().HostObject.OnClientConnected.InvokeCount == 0);
+		UTEST_TRUE_EXPR(Prerequisites->GetStats().HostObject.NotifyClientDisconnected.InvokeCount == 0);
+		BeaconClientNetStats->ReceivedControlMessages.Empty();
+		BeaconHostNetStats->ReceivedControlMessages.Empty();
+
+		// Garbage collect the host.
+		BeaconHost->DestroyBeacon();
+		BeaconHost = nullptr;
+		UTEST_TRUE_EXPR(GEngine != nullptr);
+		GEngine->ForceGarbageCollection(true);
+
+		// Tick for host connection to send close.
+		UTEST_TRUE_EXPR(BeaconUnitTest::TickOnce(*Prerequisites));
+		// Tick for client connection to see client close.
+		UTEST_TRUE_EXPR(BeaconUnitTest::TickOnce(*Prerequisites));
+		UTEST_TRUE_EXPR(BeaconClient->GetConnectionState() == EBeaconConnectionState::Invalid);
+		UTEST_TRUE_EXPR(Prerequisites->GetStats().Client.OnConnected.InvokeCount == 0);
+		UTEST_TRUE_EXPR(Prerequisites->GetStats().Client.OnFailure.InvokeCount == 1);
+		UTEST_TRUE_EXPR(Prerequisites->GetStats().Host.OnFailure.InvokeCount == 0);
+		UTEST_TRUE_EXPR(Prerequisites->GetStats().HostObject.OnClientConnected.InvokeCount == 0);
+		UTEST_TRUE_EXPR(Prerequisites->GetStats().HostObject.NotifyClientDisconnected.InvokeCount == 0);
+	}
+
+	// Garbage collect client beacon during RPC callback.
+	{
+		// Reset config and stats.
+		Prerequisites->GetConfig() = BaseConfig;
+		Prerequisites->GetStats() = BeaconUnitTest::FTestStats();
+
+		// Host setup.
+		TObjectPtr<AOnlineBeaconClient> HostUserBeaconClient;
+		TObjectPtr<AOnlineBeaconUnitTestHostObject> BeaconHostObject = Prerequisites->GetWorld()->SpawnActor<AOnlineBeaconUnitTestHostObject>(AOnlineBeaconUnitTestHostObject::StaticClass());
+		TObjectPtr<AOnlineBeaconUnitTestHost> BeaconHost = Prerequisites->GetWorld()->SpawnActor<AOnlineBeaconUnitTestHost>(AOnlineBeaconUnitTestHost::StaticClass());
+		ON_SCOPE_EXIT{ if (BeaconHostObject) { BeaconHostObject->Destroy(); } };
+		ON_SCOPE_EXIT{ if (BeaconHost) { BeaconHost->DestroyBeacon(); } };
+		UTEST_TRUE_EXPR(BeaconUnitTest::ConfigureBeacon(*Prerequisites, BeaconHost));
+		BeaconHost->ListenPort = Prerequisites->GetConfig().NetDriver.ServerListenPort;
+		BeaconHost->RegisterHost(BeaconHostObject);
+		UTEST_TRUE_EXPR(BeaconHost->InitHost());
+		BeaconHost->PauseBeaconRequests(false);
+
+		TSharedPtr<BeaconUnitTest::FNetworkStats> BeaconHostNetStats;
+		UTEST_TRUE_EXPR(BeaconUnitTest::ConfigureBeaconNetDriver(*Prerequisites, BeaconHost, &BeaconHostNetStats));
+		UTEST_TRUE_EXPR(BeaconHostNetStats.IsValid());
+
+		// Client setup.
+
+		TObjectPtr<AOnlineBeaconUnitTestClient> BeaconClient = Prerequisites->GetWorld()->SpawnActor<AOnlineBeaconUnitTestClient>(AOnlineBeaconUnitTestClient::StaticClass());
+		ON_SCOPE_EXIT{ if (BeaconClient) BeaconClient->DestroyBeacon(); };
+		UTEST_TRUE_EXPR(BeaconUnitTest::ConfigureBeacon(*Prerequisites, BeaconClient));
+
+		UTEST_TRUE_EXPR(BeaconClient->GetConnectionState() == EBeaconConnectionState::Invalid);
+		UTEST_TRUE_EXPR(BeaconUnitTest::InitClientForUser(*Prerequisites, BeaconClient, UserId));
+		UTEST_TRUE_EXPR(BeaconClient->GetConnectionState() == EBeaconConnectionState::Pending);
+
+		TSharedPtr<BeaconUnitTest::FNetworkStats> BeaconClientNetStats;
+		UTEST_TRUE_EXPR(BeaconUnitTest::ConfigureBeaconNetDriver(*Prerequisites, BeaconClient, &BeaconClientNetStats));
+		UTEST_TRUE_EXPR(BeaconClientNetStats.IsValid());
+
+		// Setup test - garbage collect the client during its OnConnected RPC from the server.
+		Prerequisites->GetConfig().Client.OnConnected.Callback = [&BeaconClient]()
+		{
+			// Garbage collect the client.
+			BeaconClient->DestroyBeacon();
+			BeaconClient = nullptr;
+		};
+		ON_SCOPE_EXIT{ Prerequisites->GetConfig().Client.OnConnected.Callback.Reset(); };
+
+		// Handshake testing.
+
+		UTEST_TRUE_EXPR(BeaconUnitTest::TickUntil(*Prerequisites, [&Prerequisites](){ return Prerequisites->GetStats().Client.OnConnected.InvokeCount == 1; }));
+		UTEST_TRUE_EXPR(Prerequisites->GetStats().Client.OnConnected.InvokeCount == 1);
+		UTEST_TRUE_EXPR(Prerequisites->GetStats().Client.OnFailure.InvokeCount == 0);
+		UTEST_TRUE_EXPR(Prerequisites->GetStats().Host.OnFailure.InvokeCount == 0);
+		UTEST_TRUE_EXPR(Prerequisites->GetStats().HostObject.OnClientConnected.InvokeCount == 1);
+		UTEST_TRUE_EXPR(Prerequisites->GetStats().HostObject.NotifyClientDisconnected.InvokeCount == 0);
+		HostUserBeaconClient = BeaconUnitTest::GetBeaconClientForUser(BeaconHostObject, UserId);
+		UTEST_TRUE_EXPR(HostUserBeaconClient != nullptr);
+		BeaconClientNetStats->ReceivedControlMessages.Empty();
+		BeaconHostNetStats->ReceivedControlMessages.Empty();
+
+		// Handshake complete, channel is open.
+
+		// Beacon destroyed during RPC - garbage collect now.
+		UTEST_TRUE_EXPR(GEngine != nullptr);
+		GEngine->ForceGarbageCollection(true);
+
+		// Tick for host connection to see client close and for delayed GC to run.
+		UTEST_TRUE_EXPR(BeaconUnitTest::TickOnce(*Prerequisites));
+
+		// Tick for host object to be notified of closure.
+		UTEST_TRUE_EXPR(BeaconUnitTest::TickOnce(*Prerequisites));
+		UTEST_TRUE_EXPR(Prerequisites->GetStats().Client.OnConnected.InvokeCount == 1);
+		UTEST_TRUE_EXPR(Prerequisites->GetStats().Client.OnFailure.InvokeCount == 0);
+		UTEST_TRUE_EXPR(Prerequisites->GetStats().Host.OnFailure.InvokeCount == 0);
+		UTEST_TRUE_EXPR(Prerequisites->GetStats().HostObject.OnClientConnected.InvokeCount == 1);
+		UTEST_TRUE_EXPR(Prerequisites->GetStats().HostObject.NotifyClientDisconnected.InvokeCount == 1);
+		UTEST_TRUE_EXPR(BeaconUnitTest::GetBeaconClientForUser(BeaconHostObject, UserId) == nullptr);
+	}
+	
+	return true;
+}
+
+IMPLEMENT_SIMPLE_AUTOMATION_TEST(FOnlineBeaconTestInvalidSocket,
+	"System.Engine.Online.OnlineSubsystemUtils.OnlineBeacon.InvalidSocket",
+	EAutomationTestFlags::EditorContext | EAutomationTestFlags::EngineFilter)
+bool FOnlineBeaconTestInvalidSocket::RunTest(const FString& Parameters)
+{
+	TSharedPtr<BeaconUnitTest::FTestPrerequisites> Prerequisites = BeaconUnitTest::FTestPrerequisites::TryCreate();
+	UTEST_TRUE_EXPR(Prerequisites.IsValid());
+
+	Prerequisites->GetConfig().NetDriver.ServerListenPort = 9999;
+
+	// Fail all netdriver initializations.
+	Prerequisites->GetConfig().NetDriver.bFailInit = true;
+
+	// Snapshot config to restore before running a test section.
+	const BeaconUnitTest::FTestConfig BaseConfig = Prerequisites->GetConfig();
+
+	FUniqueNetIdStringRef UserId = FUniqueNetIdString::Create(TEXT("User"), TEXT("UnitTest"));
+
+	// Beacon host fails init.
+	{
+		Prerequisites->GetConfig() = BaseConfig;
+		Prerequisites->GetStats() = BeaconUnitTest::FTestStats();
+
+		// Host setup.
+		TObjectPtr<AOnlineBeaconClient> HostUserBeaconClient;
+		TObjectPtr<AOnlineBeaconUnitTestHostObject> BeaconHostObject = Prerequisites->GetWorld()->SpawnActor<AOnlineBeaconUnitTestHostObject>(AOnlineBeaconUnitTestHostObject::StaticClass());
+		TObjectPtr<AOnlineBeaconUnitTestHost> BeaconHost = Prerequisites->GetWorld()->SpawnActor<AOnlineBeaconUnitTestHost>(AOnlineBeaconUnitTestHost::StaticClass());
+		ON_SCOPE_EXIT{ if (BeaconHostObject) { BeaconHostObject->Destroy(); } };
+		ON_SCOPE_EXIT{ if (BeaconHost) BeaconHost->DestroyBeacon(); };
+		UTEST_TRUE_EXPR(BeaconUnitTest::ConfigureBeacon(*Prerequisites, BeaconHost));
+		BeaconHost->ListenPort = Prerequisites->GetConfig().NetDriver.ServerListenPort;
+		BeaconHost->RegisterHost(BeaconHostObject);
+		UTEST_FALSE_EXPR(BeaconHost->InitHost());
+		UTEST_TRUE_EXPR(Prerequisites->GetStats().Client.OnConnected.InvokeCount == 0);
+		UTEST_TRUE_EXPR(Prerequisites->GetStats().Client.OnFailure.InvokeCount == 0);
+		UTEST_TRUE_EXPR(Prerequisites->GetStats().Host.OnFailure.InvokeCount == 1);
+		UTEST_TRUE_EXPR(Prerequisites->GetStats().HostObject.OnClientConnected.InvokeCount == 0);
+		UTEST_TRUE_EXPR(Prerequisites->GetStats().HostObject.NotifyClientDisconnected.InvokeCount == 0);
+	}
+
+	// Beacon client fails init.
+	{
+		Prerequisites->GetConfig() = BaseConfig;
+		Prerequisites->GetStats() = BeaconUnitTest::FTestStats();
+
+		TObjectPtr<AOnlineBeaconUnitTestClient> BeaconClient = Prerequisites->GetWorld()->SpawnActor<AOnlineBeaconUnitTestClient>(AOnlineBeaconUnitTestClient::StaticClass());
+		ON_SCOPE_EXIT{ if (BeaconClient) BeaconClient->DestroyBeacon(); };
+		UTEST_TRUE_EXPR(BeaconUnitTest::ConfigureBeacon(*Prerequisites, BeaconClient));
+		UTEST_FALSE_EXPR(BeaconUnitTest::InitClientForUser(*Prerequisites, BeaconClient, UserId));
+		UTEST_TRUE_EXPR(Prerequisites->GetStats().Client.OnConnected.InvokeCount == 0);
+		UTEST_TRUE_EXPR(Prerequisites->GetStats().Client.OnFailure.InvokeCount == 1);
+		UTEST_TRUE_EXPR(Prerequisites->GetStats().Host.OnFailure.InvokeCount == 0);
+		UTEST_TRUE_EXPR(Prerequisites->GetStats().HostObject.OnClientConnected.InvokeCount == 0);
+		UTEST_TRUE_EXPR(Prerequisites->GetStats().HostObject.NotifyClientDisconnected.InvokeCount == 0);
+	}
+
+	return true;
+}
diff --git a/Engine/Plugins/Online/OnlineSubsystemUtils/Source/OnlineSubsystemUtilsTest/Private/OnlineBeaconUnitTestClient.cpp b/Engine/Plugins/Online/OnlineSubsystemUtils/Source/OnlineSubsystemUtilsTest/Private/OnlineBeaconUnitTestClient.cpp
new file mode 100644
--- /dev/null
+++ b/Engine/Plugins/Online/OnlineSubsystemUtils/Source/OnlineSubsystemUtilsTest/Private/OnlineBeaconUnitTestClient.cpp
@@ -0,0 +1,40 @@
+// Copyright Epic Games, Inc. All Rights Reserved.
+
+#include "OnlineBeaconUnitTestClient.h"
+#include "OnlineBeaconUnitTestUtils.h"
+
+#include UE_INLINE_GENERATED_CPP_BY_NAME(OnlineBeaconUnitTestClient)
+
+AOnlineBeaconUnitTestClient::AOnlineBeaconUnitTestClient(const FObjectInitializer& ObjectInitializer) :
+	Super(ObjectInitializer)
+{
+}
+
+void AOnlineBeaconUnitTestClient::OnConnected()
+{
+	if (BeaconUnitTest::FTestStats* TestStats = BeaconUnitTest::FTestPrerequisites::GetActiveTestStats())
+	{
+		++TestStats->Client.OnConnected.InvokeCount;
+	}
+
+	Super::OnConnected();
+
+
+	if (const BeaconUnitTest::FTestConfig* TestConfig = BeaconUnitTest::FTestPrerequisites::GetActiveTestConfig())
+	{
+		if (TestConfig->Client.OnConnected.Callback)
+		{
+			TestConfig->Client.OnConnected.Callback();
+		}
+	}
+}
+
+void AOnlineBeaconUnitTestClient::OnFailure()
+{
+	if (BeaconUnitTest::FTestStats* TestStats = BeaconUnitTest::FTestPrerequisites::GetActiveTestStats())
+	{
+		++TestStats->Client.OnFailure.InvokeCount;
+	}
+
+	Super::OnFailure();
+}
diff --git a/Engine/Plugins/Online/OnlineSubsystemUtils/Source/OnlineSubsystemUtilsTest/Private/OnlineBeaconUnitTestClient.h b/Engine/Plugins/Online/OnlineSubsystemUtils/Source/OnlineSubsystemUtilsTest/Private/OnlineBeaconUnitTestClient.h
new file mode 100644
--- /dev/null
+++ b/Engine/Plugins/Online/OnlineSubsystemUtils/Source/OnlineSubsystemUtilsTest/Private/OnlineBeaconUnitTestClient.h
@@ -0,0 +1,21 @@
+// Copyright Epic Games, Inc. All Rights Reserved.
+
+#pragma once
+
+#include "OnlineBeaconClient.h"
+
+#include "OnlineBeaconUnitTestClient.generated.h"
+
+UCLASS(transient, notplaceable)
+class AOnlineBeaconUnitTestClient : public AOnlineBeaconClient
+{
+	GENERATED_UCLASS_BODY()
+
+	//~ Begin AOnlineBeaconClient Interface
+	virtual void OnConnected() override;
+	//~ End AOnlineBeaconClient Interface
+
+	//~ Begin OnlineBeacon Interface
+	virtual void OnFailure() override;
+	//~ End OnlineBeacon Interface
+};
diff --git a/Engine/Plugins/Online/OnlineSubsystemUtils/Source/OnlineSubsystemUtilsTest/Private/OnlineBeaconUnitTestHost.cpp b/Engine/Plugins/Online/OnlineSubsystemUtils/Source/OnlineSubsystemUtilsTest/Private/OnlineBeaconUnitTestHost.cpp
new file mode 100644
--- /dev/null
+++ b/Engine/Plugins/Online/OnlineSubsystemUtils/Source/OnlineSubsystemUtilsTest/Private/OnlineBeaconUnitTestHost.cpp
@@ -0,0 +1,132 @@
+// Copyright Epic Games, Inc. All Rights Reserved.
+
+#include "OnlineBeaconUnitTestHost.h"
+
+#include "Online/CoreOnline.h"
+#include "OnlineBeaconUnitTestUtils.h"
+
+#include UE_INLINE_GENERATED_CPP_BY_NAME(OnlineBeaconUnitTestHost)
+
+AOnlineBeaconUnitTestHost::AOnlineBeaconUnitTestHost(const FObjectInitializer& ObjectInitializer) :
+	Super(ObjectInitializer)
+{
+}
+
+bool AOnlineBeaconUnitTestHost::StartVerifyAuthentication(const FUniqueNetId& PlayerId, const FString& AuthenticationToken)
+{
+	const BeaconUnitTest::FTestConfig* TestConfig = BeaconUnitTest::FTestPrerequisites::GetActiveTestConfig();
+	if (TestConfig == nullptr)
+	{
+		return false;
+	}
+
+	if (!TestConfig->Auth.Method1.bEnabled)
+	{
+		PRAGMA_DISABLE_DEPRECATION_WARNINGS
+		return Super::StartVerifyAuthentication(PlayerId, AuthenticationToken);
+		PRAGMA_ENABLE_DEPRECATION_WARNINGS
+	}
+
+	if (TestConfig->Auth.Method1.bDelayDelegate)
+	{
+		return BeaconUnitTest::SetTimerForNextFrame(GetWorld(), GFrameCounter, [this, PlayerIdRef = PlayerId.AsShared()]()
+		{
+			const BeaconUnitTest::FTestConfig* TestConfig = BeaconUnitTest::FTestPrerequisites::GetActiveTestConfig();
+			PRAGMA_DISABLE_DEPRECATION_WARNINGS
+			OnAuthenticationVerificationComplete(*PlayerIdRef, TestConfig ? TestConfig->Auth.Method1.Result : FOnlineError());
+			PRAGMA_ENABLE_DEPRECATION_WARNINGS
+		});
+	}
+	else
+	{
+		PRAGMA_DISABLE_DEPRECATION_WARNINGS
+		OnAuthenticationVerificationComplete(PlayerId, TestConfig->Auth.Method1.Result);
+		PRAGMA_ENABLE_DEPRECATION_WARNINGS
+
+		return true;
+	}
+}
+
+bool AOnlineBeaconUnitTestHost::StartVerifyAuthentication(const FUniqueNetId& PlayerId, const FString& AuthenticationToken, const FOnAuthenticationVerificationCompleteDelegate& OnComplete)
+{
+	const BeaconUnitTest::FTestConfig* TestConfig = BeaconUnitTest::FTestPrerequisites::GetActiveTestConfig();
+	if (TestConfig == nullptr)
+	{
+		return false;
+	}
+
+	if (!TestConfig->Auth.Method2.bEnabled)
+	{
+		PRAGMA_DISABLE_DEPRECATION_WARNINGS
+		return Super::StartVerifyAuthentication(PlayerId, AuthenticationToken, OnComplete);
+		PRAGMA_ENABLE_DEPRECATION_WARNINGS
+	}
+
+	if (TestConfig->Auth.Method2.bDelayDelegate)
+	{
+		return BeaconUnitTest::SetTimerForNextFrame(GetWorld(), GFrameCounter, [this, OnComplete]()
+		{
+			const BeaconUnitTest::FTestConfig* TestConfig = BeaconUnitTest::FTestPrerequisites::GetActiveTestConfig();
+			OnComplete.ExecuteIfBound(TestConfig ? TestConfig->Auth.Method2.Result : FOnlineError());
+		});
+	}
+	else
+	{
+		OnComplete.ExecuteIfBound(TestConfig->Auth.Method2.Result);
+		return true;
+	}
+}
+
+bool AOnlineBeaconUnitTestHost::StartVerifyAuthentication(const FUniqueNetId& PlayerId, const FString& LoginOptions, const FString& AuthenticationToken, const FOnAuthenticationVerificationCompleteDelegate& OnComplete)
+{
+	const BeaconUnitTest::FTestConfig* TestConfig = BeaconUnitTest::FTestPrerequisites::GetActiveTestConfig();
+	if (TestConfig == nullptr)
+	{
+		return false;
+	}
+
+	if (!TestConfig->Auth.Method3.bEnabled)
+	{
+		return Super::StartVerifyAuthentication(PlayerId, LoginOptions, AuthenticationToken, OnComplete);
+	}
+
+	if (TestConfig->Auth.Method3.bDelayDelegate)
+	{
+		return BeaconUnitTest::SetTimerForNextFrame(GetWorld(), GFrameCounter, [this, OnComplete]()
+		{
+			const BeaconUnitTest::FTestConfig* TestConfig = BeaconUnitTest::FTestPrerequisites::GetActiveTestConfig();
+			OnComplete.ExecuteIfBound(TestConfig ? TestConfig->Auth.Method3.Result : FOnlineError());
+		});
+	}
+	else
+	{
+		OnComplete.ExecuteIfBound(TestConfig->Auth.Method3.Result);
+		return true;
+	}
+}
+
+bool AOnlineBeaconUnitTestHost::VerifyJoinForBeaconType(const FUniqueNetId& PlayerId, const FString& BeaconType)
+{
+	const BeaconUnitTest::FTestConfig* TestConfig = BeaconUnitTest::FTestPrerequisites::GetActiveTestConfig();
+	if (TestConfig == nullptr)
+	{
+		return false;
+	}
+
+	if (!TestConfig->Auth.Verify.bEnabled)
+	{
+		return Super::VerifyJoinForBeaconType(PlayerId, BeaconType);
+	}
+
+	return TestConfig->Auth.Verify.bResult;
+}
+
+void AOnlineBeaconUnitTestHost::OnFailure()
+{
+	if (BeaconUnitTest::FTestStats* TestStats = BeaconUnitTest::FTestPrerequisites::GetActiveTestStats())
+	{
+		++TestStats->Host.OnFailure.InvokeCount;
+	}
+
+	Super::OnFailure();
+}
diff --git a/Engine/Plugins/Online/OnlineSubsystemUtils/Source/OnlineSubsystemUtilsTest/Private/OnlineBeaconUnitTestHost.h b/Engine/Plugins/Online/OnlineSubsystemUtils/Source/OnlineSubsystemUtilsTest/Private/OnlineBeaconUnitTestHost.h
new file mode 100644
--- /dev/null
+++ b/Engine/Plugins/Online/OnlineSubsystemUtils/Source/OnlineSubsystemUtilsTest/Private/OnlineBeaconUnitTestHost.h
@@ -0,0 +1,30 @@
+// Copyright Epic Games, Inc. All Rights Reserved.
+
+#pragma once
+
+#include "OnlineBeaconHost.h"
+
+#include "OnlineBeaconUnitTestHost.generated.h"
+
+UCLASS(transient, notplaceable)
+class AOnlineBeaconUnitTestHost : public AOnlineBeaconHost
+{
+	GENERATED_UCLASS_BODY()
+
+public:
+
+	//~ Begin AOnlineBeaconHost Interface
+	PRAGMA_DISABLE_DEPRECATION_WARNINGS
+	virtual bool StartVerifyAuthentication(const FUniqueNetId& PlayerId, const FString& AuthenticationToken) override;
+	PRAGMA_ENABLE_DEPRECATION_WARNINGS
+	PRAGMA_DISABLE_DEPRECATION_WARNINGS
+	virtual bool StartVerifyAuthentication(const FUniqueNetId& PlayerId, const FString& AuthenticationToken, const FOnAuthenticationVerificationCompleteDelegate& OnComplete) override;
+	PRAGMA_ENABLE_DEPRECATION_WARNINGS
+	virtual bool StartVerifyAuthentication(const FUniqueNetId& PlayerId, const FString& LoginOptions, const FString& AuthenticationToken, const FOnAuthenticationVerificationCompleteDelegate& OnComplete) override;
+	virtual bool VerifyJoinForBeaconType(const FUniqueNetId& PlayerId, const FString& BeaconType) override;
+	//~ End AOnlineBeaconHost Interface
+
+	//~ Begin OnlineBeacon Interface
+	virtual void OnFailure() override;
+	//~ End OnlineBeacon Interface
+};
diff --git a/Engine/Plugins/Online/OnlineSubsystemUtils/Source/OnlineSubsystemUtilsTest/Private/OnlineBeaconUnitTestHostObject.cpp b/Engine/Plugins/Online/OnlineSubsystemUtils/Source/OnlineSubsystemUtilsTest/Private/OnlineBeaconUnitTestHostObject.cpp
new file mode 100644
--- /dev/null
+++ b/Engine/Plugins/Online/OnlineSubsystemUtils/Source/OnlineSubsystemUtilsTest/Private/OnlineBeaconUnitTestHostObject.cpp
@@ -0,0 +1,36 @@
+// Copyright Epic Games, Inc. All Rights Reserved.
+
+#include "OnlineBeaconUnitTestHostObject.h"
+
+#include "Engine/NetConnection.h"
+#include "OnlineBeaconUnitTestClient.h"
+#include "OnlineBeaconUnitTestUtils.h"
+
+#include UE_INLINE_GENERATED_CPP_BY_NAME(OnlineBeaconUnitTestHostObject)
+
+AOnlineBeaconUnitTestHostObject::AOnlineBeaconUnitTestHostObject(const FObjectInitializer& ObjectInitializer) :
+	Super(ObjectInitializer)
+{
+	ClientBeaconActorClass = AOnlineBeaconUnitTestClient::StaticClass();
+	BeaconTypeName = ClientBeaconActorClass->GetName();
+}
+
+void AOnlineBeaconUnitTestHostObject::OnClientConnected(AOnlineBeaconClient* NewClientActor, UNetConnection* ClientConnection)
+{
+	if (BeaconUnitTest::FTestStats* TestStats = BeaconUnitTest::FTestPrerequisites::GetActiveTestStats())
+	{
+		++TestStats->HostObject.OnClientConnected.InvokeCount;
+	}
+
+	Super::OnClientConnected(NewClientActor, ClientConnection);
+}
+
+void AOnlineBeaconUnitTestHostObject::NotifyClientDisconnected(AOnlineBeaconClient* LeavingClientActor)
+{
+	if (BeaconUnitTest::FTestStats* TestStats = BeaconUnitTest::FTestPrerequisites::GetActiveTestStats())
+	{
+		++TestStats->HostObject.NotifyClientDisconnected.InvokeCount;
+	}
+
+	Super::NotifyClientDisconnected(LeavingClientActor);
+}
diff --git a/Engine/Plugins/Online/OnlineSubsystemUtils/Source/OnlineSubsystemUtilsTest/Private/OnlineBeaconUnitTestHostObject.h b/Engine/Plugins/Online/OnlineSubsystemUtils/Source/OnlineSubsystemUtilsTest/Private/OnlineBeaconUnitTestHostObject.h
new file mode 100644
--- /dev/null
+++ b/Engine/Plugins/Online/OnlineSubsystemUtils/Source/OnlineSubsystemUtilsTest/Private/OnlineBeaconUnitTestHostObject.h
@@ -0,0 +1,19 @@
+// Copyright Epic Games, Inc. All Rights Reserved.
+
+#pragma once
+
+#include "OnlineBeaconHostObject.h"
+
+#include "OnlineBeaconUnitTestHostObject.generated.h"
+
+UCLASS(transient, notplaceable)
+class AOnlineBeaconUnitTestHostObject : public AOnlineBeaconHostObject
+{
+	GENERATED_UCLASS_BODY()
+
+public:
+	//~ Begin AOnlineBeaconHostObject Interface
+	virtual void OnClientConnected(AOnlineBeaconClient* NewClientActor, UNetConnection* ClientConnection) override;
+	virtual void NotifyClientDisconnected(AOnlineBeaconClient* LeavingClientActor) override;
+	//~ End AOnlineBeaconHostObject Interface
+};
diff --git a/Engine/Plugins/Online/OnlineSubsystemUtils/Source/OnlineSubsystemUtilsTest/Private/OnlineBeaconUnitTestNetConnection.cpp b/Engine/Plugins/Online/OnlineSubsystemUtils/Source/OnlineSubsystemUtilsTest/Private/OnlineBeaconUnitTestNetConnection.cpp
new file mode 100644
--- /dev/null
+++ b/Engine/Plugins/Online/OnlineSubsystemUtils/Source/OnlineSubsystemUtilsTest/Private/OnlineBeaconUnitTestNetConnection.cpp
@@ -0,0 +1,31 @@
+// Copyright Epic Games, Inc. All Rights Reserved.
+
+#include "OnlineBeaconUnitTestNetConnection.h"
+
+#include "OnlineBeaconUnitTestUtils.h"
+
+#include UE_INLINE_GENERATED_CPP_BY_NAME(OnlineBeaconUnitTestNetConnection)
+
+UOnlineBeaconUnitTestNetConnection::UOnlineBeaconUnitTestNetConnection(const FObjectInitializer& ObjectInitializer)
+	: Super(ObjectInitializer)
+{
+}
+
+void UOnlineBeaconUnitTestNetConnection::InitBase(UNetDriver* InDriver, class FSocket* InSocket, const FURL& InURL, EConnectionState InState, int32 InMaxPacket, int32 InPacketOverhead)
+{
+	DisableAddressResolution();
+
+	Super::InitBase(InDriver, InSocket, InURL, InState, InMaxPacket, InPacketOverhead);
+}
+
+float UOnlineBeaconUnitTestNetConnection::GetTimeoutValue()
+{
+	check(Driver);
+	if (Driver->bNoTimeouts)
+	{
+		return UE_MAX_FLT;
+	}
+
+	const bool bUseShortTimeout = (GetConnectionState() != USOCK_Pending) && (bPendingDestroy || (OwningActor && OwningActor->UseShortConnectTimeout()));
+	return bUseShortTimeout ? Driver->ConnectionTimeout : Driver->InitialConnectTimeout;
+}
diff --git a/Engine/Plugins/Online/OnlineSubsystemUtils/Source/OnlineSubsystemUtilsTest/Private/OnlineBeaconUnitTestNetConnection.h b/Engine/Plugins/Online/OnlineSubsystemUtils/Source/OnlineSubsystemUtilsTest/Private/OnlineBeaconUnitTestNetConnection.h
new file mode 100644
--- /dev/null
+++ b/Engine/Plugins/Online/OnlineSubsystemUtils/Source/OnlineSubsystemUtilsTest/Private/OnlineBeaconUnitTestNetConnection.h
@@ -0,0 +1,19 @@
+// Copyright Epic Games, Inc. All Rights Reserved.
+
+#pragma once
+
+#include "IpConnection.h"
+
+#include "OnlineBeaconUnitTestNetConnection.generated.h"
+
+UCLASS(Transient, Config=Engine)
+class UOnlineBeaconUnitTestNetConnection : public UIpConnection
+{
+	GENERATED_UCLASS_BODY()
+
+public:
+//~ Begin NetConnection Interface
+	virtual void InitBase(UNetDriver* InDriver, class FSocket* InSocket, const FURL& InURL, EConnectionState InState, int32 InMaxPacket, int32 InPacketOverhead) override;
+	virtual float GetTimeoutValue() override;
+//~ End NetConnection Interface
+};
diff --git a/Engine/Plugins/Online/OnlineSubsystemUtils/Source/OnlineSubsystemUtilsTest/Private/OnlineBeaconUnitTestNetDriver.cpp b/Engine/Plugins/Online/OnlineSubsystemUtils/Source/OnlineSubsystemUtilsTest/Private/OnlineBeaconUnitTestNetDriver.cpp
new file mode 100644
--- /dev/null
+++ b/Engine/Plugins/Online/OnlineSubsystemUtils/Source/OnlineSubsystemUtilsTest/Private/OnlineBeaconUnitTestNetDriver.cpp
@@ -0,0 +1,47 @@
+// Copyright Epic Games, Inc. All Rights Reserved.
+
+#include "OnlineBeaconUnitTestNetDriver.h"
+
+#include "OnlineBeaconUnitTestNetConnection.h"
+#include "OnlineBeaconUnitTestSocketSubsystem.h"
+#include "OnlineBeaconUnitTestUtils.h"
+
+#include UE_INLINE_GENERATED_CPP_BY_NAME(OnlineBeaconUnitTestNetDriver)
+
+UOnlineBeaconUnitTestNetDriver::UOnlineBeaconUnitTestNetDriver(const FObjectInitializer& ObjectInitializer)
+	: Super(ObjectInitializer)
+{
+}
+
+void UOnlineBeaconUnitTestNetDriver::PostInitProperties()
+{
+	Super::PostInitProperties();
+	NetConnectionClassName = UOnlineBeaconUnitTestNetConnection::StaticClass()->GetPathName();
+}
+
+bool UOnlineBeaconUnitTestNetDriver::InitBase(bool bInitAsClient, FNetworkNotify* InNotify, const FURL& URL, bool bReuseAddressAndPort, FString& Error)
+{
+	const BeaconUnitTest::FTestConfig* TestConfig = BeaconUnitTest::FTestPrerequisites::GetActiveTestConfig();
+	if (TestConfig == nullptr || TestConfig->NetDriver.bFailInit)
+	{
+		return false;
+	}
+
+	return Super::InitBase(bInitAsClient, InNotify, URL, bReuseAddressAndPort, Error);
+}
+
+void UOnlineBeaconUnitTestNetDriver::TickDispatch(float DeltaTime)
+{
+	Super::TickDispatch(DeltaTime);
+}
+
+ISocketSubsystem* UOnlineBeaconUnitTestNetDriver::GetSocketSubsystem()
+{
+	return ISocketSubsystem::Get(BeaconUnitTest::SocketSubsystemName);
+}
+
+bool UOnlineBeaconUnitTestNetDriver::IsEncryptionRequired() const
+{
+	const BeaconUnitTest::FTestConfig* TestConfig = BeaconUnitTest::FTestPrerequisites::GetActiveTestConfig();
+	return TestConfig ? TestConfig->Encryption.bEnabled : false;
+}
diff --git a/Engine/Plugins/Online/OnlineSubsystemUtils/Source/OnlineSubsystemUtilsTest/Private/OnlineBeaconUnitTestNetDriver.h b/Engine/Plugins/Online/OnlineSubsystemUtils/Source/OnlineSubsystemUtilsTest/Private/OnlineBeaconUnitTestNetDriver.h
new file mode 100644
--- /dev/null
+++ b/Engine/Plugins/Online/OnlineSubsystemUtils/Source/OnlineSubsystemUtilsTest/Private/OnlineBeaconUnitTestNetDriver.h
@@ -0,0 +1,25 @@
+// Copyright Epic Games, Inc. All Rights Reserved.
+
+#pragma once
+
+#include "IpNetDriver.h"
+
+#include "OnlineBeaconUnitTestNetDriver.generated.h"
+
+UCLASS(Transient)
+class UOnlineBeaconUnitTestNetDriver : public UIpNetDriver
+{
+	GENERATED_UCLASS_BODY()
+
+public:
+	//~ Begin UObject Interface.
+	virtual void PostInitProperties() override;
+	//~ End UObject Interface.
+
+	//~ Begin UNetDriver Interface.
+	virtual bool InitBase(bool bInitAsClient, FNetworkNotify* InNotify, const FURL& URL, bool bReuseAddressAndPort, FString& Error) override;
+	virtual void TickDispatch(float DeltaTime) override;
+	virtual ISocketSubsystem* GetSocketSubsystem() override;
+	virtual bool IsEncryptionRequired() const override;
+	//~ End UNetDriver Interface
+};
diff --git a/Engine/Plugins/Online/OnlineSubsystemUtils/Source/OnlineSubsystemUtilsTest/Private/OnlineBeaconUnitTestSocketSubsystem.cpp b/Engine/Plugins/Online/OnlineSubsystemUtils/Source/OnlineSubsystemUtilsTest/Private/OnlineBeaconUnitTestSocketSubsystem.cpp
new file mode 100644
--- /dev/null
+++ b/Engine/Plugins/Online/OnlineSubsystemUtils/Source/OnlineSubsystemUtilsTest/Private/OnlineBeaconUnitTestSocketSubsystem.cpp
@@ -0,0 +1,731 @@
+// Copyright Epic Games, Inc. All Rights Reserved.
+
+#include "OnlineBeaconUnitTestSocketSubsystem.h"
+
+#include "Algo/Find.h"
+#include "Algo/ForEach.h"
+#include "Algo/RemoveIf.h"
+#include "EngineLogs.h"
+#include "Modules/ModuleManager.h"
+#include "SocketSubsystemModule.h"
+
+DEFINE_LOG_CATEGORY_STATIC(LogUnitTestSocketSubsystem, Log, All);
+
+void FInternetAddrBeaconUnitTest::SetIp(uint32 InAddr)
+{
+}
+
+
+void FInternetAddrBeaconUnitTest::SetIp(const TCHAR* InAddr, bool& bIsValid)
+{
+}
+
+void FInternetAddrBeaconUnitTest::GetIp(uint32& OutAddr) const
+{
+	OutAddr = 0;
+}
+
+void FInternetAddrBeaconUnitTest::SetPort(int32 InPort)
+{
+	Port = InPort;
+}
+
+
+int32 FInternetAddrBeaconUnitTest::GetPort() const
+{
+	return Port;
+}
+
+void FInternetAddrBeaconUnitTest::SetRawIp(const TArray<uint8>& RawAddr)
+{
+}
+
+TArray<uint8> FInternetAddrBeaconUnitTest::GetRawIp() const
+{
+	return TArray<uint8>();
+}
+
+void FInternetAddrBeaconUnitTest::SetAnyAddress()
+{
+}
+
+void FInternetAddrBeaconUnitTest::SetBroadcastAddress()
+{
+}
+
+void FInternetAddrBeaconUnitTest::SetLoopbackAddress()
+{
+}
+
+FString FInternetAddrBeaconUnitTest::ToString(bool bAppendPort) const
+{
+	return FString::Printf(TEXT("OnlineBeaconUnitTestINetAddr:%d"), GetPort());
+}
+
+bool FInternetAddrBeaconUnitTest::operator==(const FInternetAddr& Other) const
+{
+	return Other.GetPort() == GetPort();
+}
+
+uint32 FInternetAddrBeaconUnitTest::GetTypeHash() const
+{
+	return ::GetTypeHash(Port);
+}
+
+bool FInternetAddrBeaconUnitTest::IsValid() const
+{
+	return Port != InvalidPort;
+}
+
+TSharedRef<FInternetAddr> FInternetAddrBeaconUnitTest::Clone() const
+{
+	TSharedRef<FInternetAddr> OutClone = MakeShared<FInternetAddrBeaconUnitTest>();
+	OutClone->SetPort(GetPort());
+	return OutClone;
+}
+
+FSocketBeaconUnitTest::FSocketBeaconUnitTest(ESocketType InSocketType, const FString& InSocketDescription, const FName& InSocketProtocol, const TSharedRef<FOnlineBeaconUnitTestSocketSubsystem>& InSubsystem)
+	: FSocket(InSocketType, InSocketDescription, InSocketProtocol)
+	, WeakSocketSubsystem(InSubsystem)
+{
+}
+
+FSocketBeaconUnitTest::~FSocketBeaconUnitTest()
+{
+}
+
+bool FSocketBeaconUnitTest::Shutdown(ESocketShutdownMode Mode)
+{
+	if (TSharedPtr<FOnlineBeaconUnitTestSocketSubsystem> SocketSubsystem = WeakSocketSubsystem.Pin())
+	{
+		SocketSubsystem->SetLastSocketError(ESocketErrors::SE_EOPNOTSUPP);
+	}
+	return false;
+}
+
+bool FSocketBeaconUnitTest::Close()
+{
+	return true;
+}
+
+bool FSocketBeaconUnitTest::Bind(const FInternetAddr& Addr)
+{
+	TSharedPtr<FOnlineBeaconUnitTestSocketSubsystem> SocketSubsystem = WeakSocketSubsystem.Pin();
+	if (!SocketSubsystem.IsValid())
+	{
+		return false;
+	}
+
+	const FInternetAddrBeaconUnitTest& UnitTestAddr = static_cast<const FInternetAddrBeaconUnitTest&>(Addr);
+	LocalAddress = SocketSubsystem->BindSocket(this, UnitTestAddr);
+	return LocalAddress.IsValid();
+}
+
+bool FSocketBeaconUnitTest::Connect(const FInternetAddr& Addr)
+{
+	if (TSharedPtr<FOnlineBeaconUnitTestSocketSubsystem> SocketSubsystem = WeakSocketSubsystem.Pin())
+	{
+		SocketSubsystem->SetLastSocketError(ESocketErrors::SE_EOPNOTSUPP);
+	}
+	return false;
+}
+
+bool FSocketBeaconUnitTest::Listen(int32 MaxBacklog)
+{
+	if (TSharedPtr<FOnlineBeaconUnitTestSocketSubsystem> SocketSubsystem = WeakSocketSubsystem.Pin())
+	{
+		SocketSubsystem->SetLastSocketError(ESocketErrors::SE_EOPNOTSUPP);
+	}
+	return false;
+}
+
+bool FSocketBeaconUnitTest::WaitForPendingConnection(bool& bHasPendingConnection, const FTimespan& WaitTime)
+{
+	if (TSharedPtr<FOnlineBeaconUnitTestSocketSubsystem> SocketSubsystem = WeakSocketSubsystem.Pin())
+	{
+		SocketSubsystem->SetLastSocketError(ESocketErrors::SE_EOPNOTSUPP);
+	}
+	return false;
+}
+
+bool FSocketBeaconUnitTest::HasPendingData(uint32& PendingDataSize)
+{
+	if (FUnitTestNetworkPacket* Packet = ReceiveBuffer.Peek())
+	{
+		PendingDataSize = Packet->Data.Num();
+		return true;
+	}
+
+	return false;
+}
+
+class FSocket* FSocketBeaconUnitTest::Accept(const FString& InSocketDescription)
+{
+	if (TSharedPtr<FOnlineBeaconUnitTestSocketSubsystem> SocketSubsystem = WeakSocketSubsystem.Pin())
+	{
+		SocketSubsystem->SetLastSocketError(ESocketErrors::SE_EOPNOTSUPP);
+	}
+	return nullptr;
+}
+
+class FSocket* FSocketBeaconUnitTest::Accept(FInternetAddr& OutAddr, const FString& InSocketDescription)
+{
+	if (TSharedPtr<FOnlineBeaconUnitTestSocketSubsystem> SocketSubsystem = WeakSocketSubsystem.Pin())
+	{
+		SocketSubsystem->SetLastSocketError(ESocketErrors::SE_EOPNOTSUPP);
+	}
+	return nullptr;
+}
+
+bool FSocketBeaconUnitTest::SendTo(const uint8* Data, int32 Count, int32& OutBytesSent, const FInternetAddr& Destination)
+{
+	OutBytesSent = 0;
+
+	TSharedPtr<FOnlineBeaconUnitTestSocketSubsystem> SocketSubsystem = WeakSocketSubsystem.Pin();
+	if (!SocketSubsystem.IsValid())
+	{
+		return false;
+	}
+
+	if (Count > FUnitTestNetworkPacket::MaxPacketSize)
+	{
+		UE_LOG(LogUnitTestSocketSubsystem, Warning, TEXT("[%hs] Unable to send data, data over maximum size. Amount=[%d/%d] DestinationAddress = (%s)"), __FUNCTION__, Count, FUnitTestNetworkPacket::MaxPacketSize, *Destination.ToString(true));
+		SocketSubsystem->SetLastSocketError(ESocketErrors::SE_EMSGSIZE);
+		return false;
+	}
+
+	if (Count < 0)
+	{
+		UE_LOG(LogUnitTestSocketSubsystem, Warning, TEXT("[%hs] Unable to send data, data invalid. Amount=[%d/%d] DestinationAddress = (%s)"), __FUNCTION__, Count, FUnitTestNetworkPacket::MaxPacketSize, *Destination.ToString(true));
+		SocketSubsystem->SetLastSocketError(ESocketErrors::SE_EINVAL);
+		return false;
+	}
+
+	if (!Destination.IsValid())
+	{
+		UE_LOG(LogUnitTestSocketSubsystem, Warning, TEXT("[%hs] Unable to send data, invalid destination address. DestinationAddress = (%s)"), __FUNCTION__, *Destination.ToString(true));
+		SocketSubsystem->SetLastSocketError(ESocketErrors::SE_EADDRNOTAVAIL);
+		return false;
+	}
+
+	const FInternetAddrBeaconUnitTest& DestAddress = static_cast<const FInternetAddrBeaconUnitTest&>(Destination);
+	FUnitTestNetworkPacket SendPacket;
+	SendPacket.FromAddr = LocalAddress;
+	SendPacket.ToAddr = DestAddress;
+	SendPacket.Data.SetNum(Count);
+	FMemory::Memcpy(SendPacket.Data.GetData(), Data, Count);
+	OutBytesSent = Count;
+
+	UE_LOG(LogUnitTestSocketSubsystem, VeryVerbose, TEXT("[%hs] Outbound message queued. Socket: %p, FromPort: %d, ToPort: %d"), __FUNCTION__, this, LocalAddress.GetPort(), Destination.GetPort());
+
+	SendBuffer.Enqueue(MoveTemp(SendPacket));
+	return true;
+}
+
+bool FSocketBeaconUnitTest::Send(const uint8* Data, int32 Count, int32& BytesSent)
+{
+	BytesSent = 0;
+	if (TSharedPtr<FOnlineBeaconUnitTestSocketSubsystem> SocketSubsystem = WeakSocketSubsystem.Pin())
+	{
+		SocketSubsystem->SetLastSocketError(ESocketErrors::SE_EOPNOTSUPP);
+	}
+	return false;
+}
+
+bool FSocketBeaconUnitTest::RecvFrom(uint8* Data, int32 BufferSize, int32& BytesRead, FInternetAddr& Source, ESocketReceiveFlags::Type Flags)
+{
+	BytesRead = 0;
+
+	TSharedPtr<FOnlineBeaconUnitTestSocketSubsystem> SocketSubsystem = WeakSocketSubsystem.Pin();
+	if (!SocketSubsystem.IsValid())
+	{
+		return false;
+	}
+
+	if (BufferSize < 0)
+	{
+		UE_LOG(LogUnitTestSocketSubsystem, Error, TEXT("[%hs] Unable to receive data, receiving buffer was invalid. BufferSize = (%d)"), __FUNCTION__, BufferSize);
+
+		SocketSubsystem->SetLastSocketError(ESocketErrors::SE_EINVAL);
+		return false;
+	}
+
+	if (Flags != ESocketReceiveFlags::None)
+	{
+		// We do not support peaking / blocking until a packet comes
+		UE_LOG(LogUnitTestSocketSubsystem, Error, TEXT("[%hs] Socket receive flags (%d) are not supported"), __FUNCTION__, int32(Flags));
+
+		SocketSubsystem->SetLastSocketError(ESocketErrors::SE_EOPNOTSUPP);
+		return false;
+	}
+
+	const FUnitTestNetworkPacket* ReceivePacket = ReceiveBuffer.Peek();
+	if (ReceivePacket == nullptr)
+	{
+		SocketSubsystem->SetLastSocketError(ESocketErrors::SE_EWOULDBLOCK);
+		return false;
+	}
+
+	if (!EnumHasAnyFlags(UnitTestFlags, BeaconUnitTest::ESocketFlags::RecvEnabled))
+	{
+		SocketSubsystem->SetLastSocketError(ESocketErrors::SE_EWOULDBLOCK);
+		return false;
+	}
+
+	UE_LOG(LogUnitTestSocketSubsystem, VeryVerbose, TEXT("[%hs] Inbound message received. Socket: %p, FromPort: %d, ToPort: %d"), __FUNCTION__, this, ReceivePacket->FromAddr.GetPort(), LocalAddress.GetPort());
+
+	FInternetAddrBeaconUnitTest& SourceAddress = static_cast<FInternetAddrBeaconUnitTest&>(Source);
+	SourceAddress = ReceivePacket->FromAddr;
+	BytesRead = FMath::Min(BufferSize, ReceivePacket->Data.Num());
+	FMemory::Memcpy(Data, ReceivePacket->Data.GetData(), BytesRead);
+	ReceiveBuffer.Pop();
+
+	return true;
+}
+
+bool FSocketBeaconUnitTest::Recv(uint8* Data, int32 BufferSize, int32& BytesRead, ESocketReceiveFlags::Type Flags)
+{
+	BytesRead = 0;
+	if (TSharedPtr<FOnlineBeaconUnitTestSocketSubsystem> SocketSubsystem = WeakSocketSubsystem.Pin())
+	{
+		SocketSubsystem->SetLastSocketError(ESocketErrors::SE_EOPNOTSUPP);
+	}
+	return false;
+}
+
+bool FSocketBeaconUnitTest::Wait(ESocketWaitConditions::Type Condition, FTimespan WaitTime)
+{
+	if (TSharedPtr<FOnlineBeaconUnitTestSocketSubsystem> SocketSubsystem = WeakSocketSubsystem.Pin())
+	{
+		SocketSubsystem->SetLastSocketError(ESocketErrors::SE_EOPNOTSUPP);
+	}
+	return false;
+}
+
+ESocketConnectionState FSocketBeaconUnitTest::GetConnectionState()
+{
+	return ESocketConnectionState::SCS_NotConnected;
+}
+
+void FSocketBeaconUnitTest::GetAddress(FInternetAddr& OutAddr)
+{
+	OutAddr = LocalAddress;
+}
+
+bool FSocketBeaconUnitTest::GetPeerAddress(FInternetAddr& OutAddr)
+{
+	if (TSharedPtr<FOnlineBeaconUnitTestSocketSubsystem> SocketSubsystem = WeakSocketSubsystem.Pin())
+	{
+		SocketSubsystem->SetLastSocketError(ESocketErrors::SE_EOPNOTSUPP);
+	}
+	return false;
+}
+
+bool FSocketBeaconUnitTest::SetNonBlocking(bool bIsNonBlocking)
+{
+	return true;
+}
+
+bool FSocketBeaconUnitTest::SetBroadcast(bool bAllowBroadcast)
+{
+	return true;
+}
+
+bool FSocketBeaconUnitTest::SetNoDelay(bool bIsNoDelay)
+{
+	return true;
+}
+
+bool FSocketBeaconUnitTest::JoinMulticastGroup(const FInternetAddr& GroupAddress)
+{
+	if (TSharedPtr<FOnlineBeaconUnitTestSocketSubsystem> SocketSubsystem = WeakSocketSubsystem.Pin())
+	{
+		SocketSubsystem->SetLastSocketError(ESocketErrors::SE_EOPNOTSUPP);
+	}
+	return false;
+}
+
+bool FSocketBeaconUnitTest::JoinMulticastGroup(const FInternetAddr& GroupAddress, const FInternetAddr& InterfaceAddress)
+{
+	if (TSharedPtr<FOnlineBeaconUnitTestSocketSubsystem> SocketSubsystem = WeakSocketSubsystem.Pin())
+	{
+		SocketSubsystem->SetLastSocketError(ESocketErrors::SE_EOPNOTSUPP);
+	}
+	return false;
+}
+
+bool FSocketBeaconUnitTest::LeaveMulticastGroup(const FInternetAddr& GroupAddress)
+{
+	if (TSharedPtr<FOnlineBeaconUnitTestSocketSubsystem> SocketSubsystem = WeakSocketSubsystem.Pin())
+	{
+		SocketSubsystem->SetLastSocketError(ESocketErrors::SE_EOPNOTSUPP);
+	}
+	return false;
+}
+
+bool FSocketBeaconUnitTest::LeaveMulticastGroup(const FInternetAddr& GroupAddress, const FInternetAddr& InterfaceAddress)
+{
+	if (TSharedPtr<FOnlineBeaconUnitTestSocketSubsystem> SocketSubsystem = WeakSocketSubsystem.Pin())
+	{
+		SocketSubsystem->SetLastSocketError(ESocketErrors::SE_EOPNOTSUPP);
+	}
+	return false;
+}
+
+bool FSocketBeaconUnitTest::SetMulticastLoopback(bool bLoopback)
+{
+	if (TSharedPtr<FOnlineBeaconUnitTestSocketSubsystem> SocketSubsystem = WeakSocketSubsystem.Pin())
+	{
+		SocketSubsystem->SetLastSocketError(ESocketErrors::SE_EOPNOTSUPP);
+	}
+	return false;
+}
+
+bool FSocketBeaconUnitTest::SetMulticastTtl(uint8 TimeToLive)
+{
+	if (TSharedPtr<FOnlineBeaconUnitTestSocketSubsystem> SocketSubsystem = WeakSocketSubsystem.Pin())
+	{
+		SocketSubsystem->SetLastSocketError(ESocketErrors::SE_EOPNOTSUPP);
+	}
+	return false;
+}
+
+bool FSocketBeaconUnitTest::SetMulticastInterface(const FInternetAddr& InterfaceAddress)
+{
+	if (TSharedPtr<FOnlineBeaconUnitTestSocketSubsystem> SocketSubsystem = WeakSocketSubsystem.Pin())
+	{
+		SocketSubsystem->SetLastSocketError(ESocketErrors::SE_EOPNOTSUPP);
+	}
+	return false;
+}
+
+bool FSocketBeaconUnitTest::SetReuseAddr(bool bAllowReuse)
+{
+	return true;
+}
+
+bool FSocketBeaconUnitTest::SetLinger(bool bShouldLinger, int32 Timeout)
+{
+	return true;
+}
+
+bool FSocketBeaconUnitTest::SetRecvErr(bool bUseErrorQueue)
+{
+	return true;
+}
+
+bool FSocketBeaconUnitTest::SetSendBufferSize(int32 Size, int32& NewSize)
+{
+	return true;
+}
+
+bool FSocketBeaconUnitTest::SetReceiveBufferSize(int32 Size, int32& NewSize)
+{
+	return true;
+}
+
+int32 FSocketBeaconUnitTest::GetPortNo()
+{
+	return LocalAddress.GetPort();
+}
+
+void FSocketBeaconUnitTest::SetUnitTestFlags(BeaconUnitTest::ESocketFlags InUnitTestFlags)
+{
+	UnitTestFlags = InUnitTestFlags;
+}
+
+void FSocketBeaconUnitTest::FlushSendBuffer()
+{
+	TSharedPtr<FOnlineBeaconUnitTestSocketSubsystem> SocketSubsystem = WeakSocketSubsystem.Pin();
+	if (!SocketSubsystem.IsValid())
+	{
+		return;
+	}
+
+	if (EnumHasAnyFlags(UnitTestFlags, BeaconUnitTest::ESocketFlags::SendEnabled) && !SendBuffer.IsEmpty())
+	{
+		UE_LOG(LogUnitTestSocketSubsystem, VeryVerbose, TEXT("[%hs] Flushing send buffer. Socket: %p, FromPort: %d"), __FUNCTION__, this, LocalAddress.GetPort());
+
+		while (FUnitTestNetworkPacket* SendPacket = SendBuffer.Peek())
+		{
+			SocketSubsystem->DispatchTestPacket(MoveTemp(*SendPacket));
+			SendBuffer.Pop();
+		}
+	}
+}
+
+FOnlineBeaconUnitTestSocketSubsystem* FOnlineBeaconUnitTestSocketSubsystem::Singleton = nullptr;
+
+FOnlineBeaconUnitTestSocketSubsystem::FOnlineBeaconUnitTestSocketSubsystem()
+	: LastSocketError(SE_NO_ERROR)
+{
+}
+
+FOnlineBeaconUnitTestSocketSubsystem::~FOnlineBeaconUnitTestSocketSubsystem()
+{
+}
+
+FOnlineBeaconUnitTestSocketSubsystem* FOnlineBeaconUnitTestSocketSubsystem::Get()
+{
+	return Singleton;
+}
+
+bool FOnlineBeaconUnitTestSocketSubsystem::Init(FString& Error)
+{
+	FSocketSubsystemModule& SocketSubsystem = FModuleManager::LoadModuleChecked<FSocketSubsystemModule>("Sockets");
+	SocketSubsystem.RegisterSocketSubsystem(BeaconUnitTest::SocketSubsystemName, this, false);
+
+	ensure(Singleton == nullptr);
+	Singleton = this;
+
+	return true;
+}
+
+void FOnlineBeaconUnitTestSocketSubsystem::Shutdown()
+{
+	Singleton = nullptr;
+
+	if (FSocketSubsystemModule* SocketSubsystem = FModuleManager::GetModulePtr<FSocketSubsystemModule>("Sockets"))
+	{
+		SocketSubsystem->UnregisterSocketSubsystem(BeaconUnitTest::SocketSubsystemName);
+	}
+}
+
+FSocket* FOnlineBeaconUnitTestSocketSubsystem::CreateSocket(const FName& SocketType, const FString& SocketDescription, const FName& ProtocolType)
+{
+	if (!ensure(SocketType == NAME_DGram))
+	{
+		return nullptr;
+	}
+
+	FSocketInfo& SocketInfo = Sockets.Emplace_GetRef(MakeShared<FSocketBeaconUnitTest>(SOCKTYPE_Datagram, SocketDescription, ProtocolType, AsShared()), FInternetAddrBeaconUnitTest(), false);
+	return SocketInfo.Socket.Get();
+}
+
+void FOnlineBeaconUnitTestSocketSubsystem::DestroySocket(class FSocket* Socket)
+{
+	DestroySocket(Socket, true);
+}
+
+FAddressInfoResult FOnlineBeaconUnitTestSocketSubsystem::GetAddressInfo(const TCHAR* HostName, const TCHAR* ServiceName,
+	EAddressInfoFlags QueryFlags,
+	const FName ProtocolTypeName,
+	ESocketType SocketType)
+{
+	FAddressInfoResult AddrData = FAddressInfoResult(HostName, ServiceName);
+	AddrData.QueryHostName = HostName;
+	AddrData.ReturnCode = SE_NO_ERROR;
+
+	const int32 PortValue = (AddrData.QueryServiceName.IsNumeric()) ? FCString::Atoi(ServiceName) : -1;
+	TSharedRef<FInternetAddrBeaconUnitTest> NewAddress = MakeShared<FInternetAddrBeaconUnitTest>();
+	NewAddress->SetPort(PortValue);
+	AddrData.Results.Add(FAddressInfoResultData(NewAddress, 0, NAME_None, SocketType));
+
+	return AddrData;
+}
+
+TSharedPtr<FInternetAddr> FOnlineBeaconUnitTestSocketSubsystem::GetAddressFromString(const FString& IPAddress)
+{
+	return nullptr;
+}
+
+class FResolveInfo* FOnlineBeaconUnitTestSocketSubsystem::GetHostByName(const ANSICHAR* HostName)
+{
+	return nullptr;
+}
+
+bool FOnlineBeaconUnitTestSocketSubsystem::RequiresChatDataBeSeparate()
+{
+	return false;
+}
+
+bool FOnlineBeaconUnitTestSocketSubsystem::RequiresEncryptedPackets()
+{
+	return false;
+}
+
+bool FOnlineBeaconUnitTestSocketSubsystem::GetHostName(FString& HostName)
+{
+	return false;
+}
+
+TSharedRef<FInternetAddr> FOnlineBeaconUnitTestSocketSubsystem::CreateInternetAddr()
+{
+	return MakeShared<FInternetAddrBeaconUnitTest>();
+}
+
+bool FOnlineBeaconUnitTestSocketSubsystem::HasNetworkDevice()
+{
+	return true;
+}
+
+const TCHAR* FOnlineBeaconUnitTestSocketSubsystem::GetSocketAPIName() const
+{
+	return TEXT("OnlineBeaconUnitTest");
+}
+
+ESocketErrors FOnlineBeaconUnitTestSocketSubsystem::GetLastErrorCode()
+{
+	return TranslateErrorCode(LastSocketError);
+}
+
+ESocketErrors FOnlineBeaconUnitTestSocketSubsystem::TranslateErrorCode(int32 Code)
+{
+	return static_cast<ESocketErrors>(Code);
+}
+
+bool FOnlineBeaconUnitTestSocketSubsystem::GetLocalAdapterAddresses(TArray<TSharedPtr<FInternetAddr> >& OutAdresses)
+{
+	OutAdresses.Emplace(MakeShared<FInternetAddrBeaconUnitTest>());
+	return true;
+}
+
+TArray<TSharedRef<FInternetAddr>> FOnlineBeaconUnitTestSocketSubsystem::GetLocalBindAddresses()
+{
+	return TArray<TSharedRef<FInternetAddr>>({MakeShared<FInternetAddrBeaconUnitTest>()});
+}
+
+bool FOnlineBeaconUnitTestSocketSubsystem::IsSocketWaitSupported() const
+{
+	return false;
+}
+
+void FOnlineBeaconUnitTestSocketSubsystem::SetLastSocketError(const ESocketErrors NewSocketError)
+{
+	LastSocketError = NewSocketError;
+}
+
+void FOnlineBeaconUnitTestSocketSubsystem::DispatchTestPacket(FUnitTestNetworkPacket&& Packet)
+{
+	if (FSocketInfo* SocketInfo = FindUnitTestSocketInfo(Packet.ToAddr))
+	{
+		UE_LOG(LogUnitTestSocketSubsystem, VeryVerbose, TEXT("[%hs] Dispatching packet to socket. ToSocket: %p, FromPort: %d, ToPort: %d"), __FUNCTION__, SocketInfo->Socket.Get(), Packet.FromAddr.GetPort(), Packet.ToAddr.GetPort());
+		SocketInfo->Socket->ReceiveBuffer.Enqueue(MoveTemp(Packet));
+	}
+	else
+	{
+		UE_LOG(LogUnitTestSocketSubsystem, VeryVerbose, TEXT("[%hs] Failed to find socket for destination. FromPort: %d, ToPort: %d"), __FUNCTION__, Packet.FromAddr.GetPort(), Packet.ToAddr.GetPort());
+	}
+}
+
+void FOnlineBeaconUnitTestSocketSubsystem::FlushSendBuffers()
+{
+	TArray<FSocketInfo> SocketsReadyForDestroy;
+
+	for (FSocketInfo& SocketInfo : Sockets)
+	{
+		if (SocketInfo.BoundAddress.IsValid())
+		{
+			SocketInfo.Socket->FlushSendBuffer();
+		}
+		
+		if (SocketInfo.bDestroyPendingFlush)
+		{
+			SocketsReadyForDestroy.Add(SocketInfo);
+		}
+	}
+
+	for (FSocketInfo& SocketInfo : SocketsReadyForDestroy)
+	{
+		DestroySocket(SocketInfo.Socket.Get(), false);
+	}
+}
+
+FInternetAddrBeaconUnitTest FOnlineBeaconUnitTestSocketSubsystem::BindSocket(FSocketBeaconUnitTest* Socket, const FInternetAddrBeaconUnitTest& RequestedAddress)
+{
+	FInternetAddrBeaconUnitTest ResolvedAddress = RequestedAddress;
+
+	// Lookup the internal pointer for bookkeeping.
+	FSocketInfo* SocketInfo = FindUnitTestSocketInfo(Socket);
+	if (SocketInfo == nullptr)
+	{
+		return FInternetAddrBeaconUnitTest();
+	}
+
+	// Check whether an ephemeral port should be assigned.
+	if (!ResolvedAddress.IsValid())
+	{
+		ResolvedAddress.SetPort(AllocateEphemeralPort());
+	}
+
+	UE_LOG(LogUnitTestSocketSubsystem, Verbose, TEXT("[%hs] Binding socket to requested address. Socket: %p, Port: %d"), __FUNCTION__, Socket, RequestedAddress.GetPort());
+
+	// Unbind current socket if already bound.
+	UnbindSocket(*SocketInfo);
+
+	// Remove previous use of address discarding any messages pending flush to prevent tests interfering with each other.
+	if (FSocketInfo* ExistingBind = FindUnitTestSocketInfo(ResolvedAddress))
+	{
+		UE_LOG(LogUnitTestSocketSubsystem, Verbose, TEXT("[%hs] Unbinding existing socket from requested address. Socket: %p, Port: %d"), __FUNCTION__, ExistingBind->Socket.Get(), RequestedAddress.GetPort());
+		UnbindSocket(*ExistingBind);
+	}
+
+	SocketInfo->BoundAddress = ResolvedAddress;
+	return ResolvedAddress;
+}
+
+int32 FOnlineBeaconUnitTestSocketSubsystem::AllocateEphemeralPort()
+{
+	return NextEphemeralPort++;
+}
+
+void FOnlineBeaconUnitTestSocketSubsystem::DestroySocket(class FSocket* Socket, bool bFlushTransmit)
+{
+	if (bFlushTransmit)
+	{
+		if (FSocketInfo* SocketInfo = FindUnitTestSocketInfo(Socket))
+		{
+			const bool bHasMessagesPendingTransmit = !SocketInfo->Socket->SendBuffer.IsEmpty() && EnumHasAnyFlags(SocketInfo->Socket->UnitTestFlags, BeaconUnitTest::ESocketFlags::SendEnabled);
+			if (bHasMessagesPendingTransmit)
+			{
+				UE_LOG(LogUnitTestSocketSubsystem, Verbose, TEXT("[%hs] Setting socket for pending destroy due to buffered outbound messages. Socket: %p, Port: %d"), __FUNCTION__, Socket, Socket->GetPortNo());
+				SocketInfo->bDestroyPendingFlush = true;
+			}
+			else
+			{
+				DestroySocket(Socket, false);
+			}
+		}
+	}
+	else
+	{
+		UE_LOG(LogUnitTestSocketSubsystem, Verbose, TEXT("[%hs] Destroying socket. Socket: %p, Port: %d"), __FUNCTION__, Socket, Socket->GetPortNo());
+
+		Sockets.SetNum(Algo::RemoveIf(Sockets, [Socket](const FSocketInfo& SocketInfo)
+		{
+			return SocketInfo.Socket.Get() == Socket;
+		}));
+	}
+}
+
+void FOnlineBeaconUnitTestSocketSubsystem::UnbindSocket(FSocketInfo& SocketInfo)
+{
+	SocketInfo.Socket->SendBuffer.Empty();
+	SocketInfo.BoundAddress = FInternetAddrBeaconUnitTest();
+}
+
+FOnlineBeaconUnitTestSocketSubsystem::FSocketInfo* FOnlineBeaconUnitTestSocketSubsystem::FindUnitTestSocketInfo(const FSocket* Socket)
+{
+	return Algo::FindByPredicate(Sockets, [Socket](const FSocketInfo& SocketInfo)
+	{
+		return SocketInfo.Socket.Get() == Socket;
+	});
+}
+
+const FOnlineBeaconUnitTestSocketSubsystem::FSocketInfo* FOnlineBeaconUnitTestSocketSubsystem::FindUnitTestSocketInfo(const FSocket* Socket) const
+{
+	return const_cast<FOnlineBeaconUnitTestSocketSubsystem*>(this)->FindUnitTestSocketInfo(Socket);
+}
+
+FOnlineBeaconUnitTestSocketSubsystem::FSocketInfo* FOnlineBeaconUnitTestSocketSubsystem::FindUnitTestSocketInfo(const FInternetAddrBeaconUnitTest& Address)
+{
+	return Algo::FindByPredicate(Sockets, [&Address](const FSocketInfo& SocketInfo)
+	{
+		return SocketInfo.BoundAddress == Address;
+	});
+}
+
+const FOnlineBeaconUnitTestSocketSubsystem::FSocketInfo* FOnlineBeaconUnitTestSocketSubsystem::FindUnitTestSocketInfo(const FInternetAddrBeaconUnitTest& Address) const
+{
+	return const_cast<FOnlineBeaconUnitTestSocketSubsystem*>(this)->FindUnitTestSocketInfo(Address);
+}
diff --git a/Engine/Plugins/Online/OnlineSubsystemUtils/Source/OnlineSubsystemUtilsTest/Private/OnlineBeaconUnitTestSocketSubsystem.h b/Engine/Plugins/Online/OnlineSubsystemUtils/Source/OnlineSubsystemUtilsTest/Private/OnlineBeaconUnitTestSocketSubsystem.h
new file mode 100644
--- /dev/null
+++ b/Engine/Plugins/Online/OnlineSubsystemUtils/Source/OnlineSubsystemUtilsTest/Private/OnlineBeaconUnitTestSocketSubsystem.h
@@ -0,0 +1,193 @@
+// Copyright Epic Games, Inc. All Rights Reserved.
+
+#pragma once
+
+#include "Containers/Queue.h"
+#include "OnlineBeaconUnitTestUtils.h"
+#include "SocketSubsystem.h"
+#include "Sockets.h"
+
+class FOnlineBeaconUnitTestSocketSubsystem;
+
+class FInternetAddrBeaconUnitTest : public FInternetAddr
+{
+public:
+
+	FInternetAddrBeaconUnitTest()
+	{
+	}
+
+	//~ Begin FInternetAddr Interface
+	virtual void SetIp(uint32 InAddr) override;
+	virtual void SetIp(const TCHAR* InAddr, bool& bIsValid) override;
+	virtual void GetIp(uint32& OutAddr) const override;
+	virtual void SetPort(int32 InPort) override;
+	virtual int32 GetPort() const override;
+	virtual void SetRawIp(const TArray<uint8>& RawAddr) override;
+	virtual TArray<uint8> GetRawIp() const override;
+	virtual void SetAnyAddress() override;
+	virtual void SetBroadcastAddress() override;
+	virtual void SetLoopbackAddress() override;
+	virtual FString ToString(bool bAppendPort) const override;
+	virtual bool operator==(const FInternetAddr& Other) const override;
+	virtual uint32 GetTypeHash() const override;
+	virtual bool IsValid() const override;
+	virtual TSharedRef<FInternetAddr> Clone() const override;
+	//~ End FInternetAddr Interface
+
+	bool operator==(const FInternetAddrBeaconUnitTest& Other) const
+	{
+		return Other.Port == Port;
+	}
+	bool operator!=(const FInternetAddrBeaconUnitTest& Other) const
+	{
+		return !(FInternetAddrBeaconUnitTest::operator==(Other));
+	}
+
+	friend uint32 GetTypeHash(const FInternetAddrBeaconUnitTest& A)
+	{
+		return A.GetTypeHash();
+	}
+
+private:
+	static constexpr int32 InvalidPort = 0;
+	int32 Port = InvalidPort;
+};
+
+struct FUnitTestNetworkPacket
+{
+	static constexpr int32 MaxPacketSize = 65535;
+
+	FInternetAddrBeaconUnitTest FromAddr;
+	FInternetAddrBeaconUnitTest ToAddr;
+	TArray<uint8> Data;
+};
+
+class FSocketBeaconUnitTest : public FSocket
+{
+public:
+	FSocketBeaconUnitTest(ESocketType InSocketType, const FString& InSocketDescription, const FName& InSocketProtocol, const TSharedRef<FOnlineBeaconUnitTestSocketSubsystem>& InSubsystem);
+	virtual ~FSocketBeaconUnitTest();
+
+	//~ Begin FSocket Interface
+	virtual bool Shutdown(ESocketShutdownMode Mode) override;
+	virtual bool Close() override;
+	virtual bool Bind(const FInternetAddr& Addr) override;
+	virtual bool Connect(const FInternetAddr& Addr) override;
+	virtual bool Listen(int32 MaxBacklog) override;
+	virtual bool WaitForPendingConnection(bool& bHasPendingConnection, const FTimespan& WaitTime) override;
+	virtual bool HasPendingData(uint32& PendingDataSize) override;
+	virtual class FSocket* Accept(const FString& InSocketDescription) override;
+	virtual class FSocket* Accept(FInternetAddr& OutAddr, const FString& InSocketDescription) override;
+	virtual bool SendTo(const uint8* Data, int32 Count, int32& BytesSent, const FInternetAddr& Destination);
+	virtual bool Send(const uint8* Data, int32 Count, int32& BytesSent);
+	virtual bool RecvFrom(uint8* Data, int32 BufferSize, int32& BytesRead, FInternetAddr& Source, ESocketReceiveFlags::Type Flags = ESocketReceiveFlags::None);
+	virtual bool Recv(uint8* Data, int32 BufferSize, int32& BytesRead, ESocketReceiveFlags::Type Flags = ESocketReceiveFlags::None);
+	virtual bool Wait(ESocketWaitConditions::Type Condition, FTimespan WaitTime) override;
+	virtual ESocketConnectionState GetConnectionState() override;
+	virtual void GetAddress(FInternetAddr& OutAddr) override;
+	virtual bool GetPeerAddress(FInternetAddr& OutAddr) override;
+	virtual bool SetNonBlocking(bool bIsNonBlocking = true) override;
+	virtual bool SetBroadcast(bool bAllowBroadcast = true) override;
+	virtual bool SetNoDelay(bool bIsNoDelay = true) override;
+	virtual bool JoinMulticastGroup(const FInternetAddr& GroupAddress) override;
+	virtual bool JoinMulticastGroup(const FInternetAddr& GroupAddress, const FInternetAddr& InterfaceAddress) override;
+	virtual bool LeaveMulticastGroup(const FInternetAddr& GroupAddress) override;
+	virtual bool LeaveMulticastGroup(const FInternetAddr& GroupAddress, const FInternetAddr& InterfaceAddress) override;
+	virtual bool SetMulticastLoopback(bool bLoopback) override;
+	virtual bool SetMulticastTtl(uint8 TimeToLive) override;
+	virtual bool SetMulticastInterface(const FInternetAddr& InterfaceAddress) override;
+	virtual bool SetReuseAddr(bool bAllowReuse = true) override;
+	virtual bool SetLinger(bool bShouldLinger = true, int32 Timeout = 0) override;
+	virtual bool SetRecvErr(bool bUseErrorQueue = true) override;
+	virtual bool SetSendBufferSize(int32 Size, int32& NewSize) override;
+	virtual bool SetReceiveBufferSize(int32 Size, int32& NewSize) override;
+	virtual int32 GetPortNo() override;
+	//~ End FSocket Interface
+
+	void SetUnitTestFlags(BeaconUnitTest::ESocketFlags UnitTestFlags);
+	void FlushSendBuffer();
+
+	/** Reference to our subsystem */
+	TWeakPtr<FOnlineBeaconUnitTestSocketSubsystem> WeakSocketSubsystem;
+
+	/** Our local address; session/port will be invalid when not bound */
+	FInternetAddrBeaconUnitTest LocalAddress;
+
+	BeaconUnitTest::ESocketFlags UnitTestFlags = BeaconUnitTest::ESocketFlags::Default;
+
+	TQueue<FUnitTestNetworkPacket> SendBuffer;
+	TQueue<FUnitTestNetworkPacket> ReceiveBuffer;
+};
+
+class FOnlineBeaconUnitTestSocketSubsystem : public ISocketSubsystem, public TSharedFromThis<FOnlineBeaconUnitTestSocketSubsystem>
+{
+public:
+	FOnlineBeaconUnitTestSocketSubsystem();
+	virtual ~FOnlineBeaconUnitTestSocketSubsystem();
+
+	static FOnlineBeaconUnitTestSocketSubsystem* Get();
+
+	//~ Begin ISocketSubsystem Interface
+	virtual bool Init(FString& Error) override;
+	virtual void Shutdown() override;
+	virtual FSocket* CreateSocket(const FName& SocketType, const FString& SocketDescription, const FName& ProtocolName) override;
+	virtual void DestroySocket(class FSocket* Socket) override;
+	virtual FAddressInfoResult GetAddressInfo(const TCHAR* HostName, const TCHAR* ServiceName = nullptr,
+		EAddressInfoFlags QueryFlags = EAddressInfoFlags::Default,
+		const FName ProtocolTypeName = NAME_None,
+		ESocketType SocketType = ESocketType::SOCKTYPE_Unknown) override;
+	virtual TSharedPtr<FInternetAddr> GetAddressFromString(const FString& IPAddress) override;
+	virtual class FResolveInfo* GetHostByName(const ANSICHAR* HostName) override;
+	virtual bool RequiresChatDataBeSeparate() override;
+	virtual bool RequiresEncryptedPackets() override;
+	virtual bool GetHostName(FString& HostName) override;
+	virtual TSharedRef<FInternetAddr> CreateInternetAddr() override;
+	virtual bool HasNetworkDevice() override;
+	virtual const TCHAR* GetSocketAPIName() const override;
+	virtual ESocketErrors GetLastErrorCode() override;
+	virtual ESocketErrors TranslateErrorCode(int32 Code) override;
+	virtual bool GetLocalAdapterAddresses(TArray<TSharedPtr<FInternetAddr> >& OutAdresses) override;
+	virtual TArray<TSharedRef<FInternetAddr>> GetLocalBindAddresses() override;
+	virtual bool IsSocketWaitSupported() const override;
+	//~ End ISocketSubsystem Interface
+
+	void SetLastSocketError(const ESocketErrors NewSocketError);
+
+	void DispatchTestPacket(FUnitTestNetworkPacket&& Packet);
+
+	void FlushSendBuffers();
+
+	/*
+	 * Bind the socket to the requested test address. If the port is unset an ephemeral port will be bound.
+	 * Once bound, the socket will be flushed for pending IO when Flush is called.
+	 * Returns the resolved address.
+	 */
+	FInternetAddrBeaconUnitTest BindSocket(FSocketBeaconUnitTest* Socket, const FInternetAddrBeaconUnitTest& RequestedAddress);
+
+private:
+
+	struct FSocketInfo
+	{
+		TSharedPtr<FSocketBeaconUnitTest> Socket;
+		FInternetAddrBeaconUnitTest BoundAddress;
+		bool bDestroyPendingFlush = false;
+	};
+
+	int32 AllocateEphemeralPort();
+	void DestroySocket(class FSocket* Socket, bool bFlushTransmit);
+	void UnbindSocket(FSocketInfo& SocketInfo);
+
+	FSocketInfo* FindUnitTestSocketInfo(const FSocket* Socket);
+	const FSocketInfo* FindUnitTestSocketInfo(const FSocket* Socket) const;
+	FSocketInfo* FindUnitTestSocketInfo(const FInternetAddrBeaconUnitTest& Address);
+	const FSocketInfo* FindUnitTestSocketInfo(const FInternetAddrBeaconUnitTest& Address) const;
+
+	static FOnlineBeaconUnitTestSocketSubsystem* Singleton;
+	ESocketErrors LastSocketError;
+	TArray<FSocketInfo> Sockets;
+
+	// Start assignment outside of normal port range.
+	static constexpr int32 EphemeralPortStart = 65536;
+	int32 NextEphemeralPort = EphemeralPortStart;
+};
diff --git a/Engine/Plugins/Online/OnlineSubsystemUtils/Source/OnlineSubsystemUtilsTest/Private/OnlineBeaconUnitTestUtils.cpp b/Engine/Plugins/Online/OnlineSubsystemUtils/Source/OnlineSubsystemUtilsTest/Private/OnlineBeaconUnitTestUtils.cpp
new file mode 100644
--- /dev/null
+++ b/Engine/Plugins/Online/OnlineSubsystemUtils/Source/OnlineSubsystemUtilsTest/Private/OnlineBeaconUnitTestUtils.cpp
@@ -0,0 +1,546 @@
+// Copyright Epic Games, Inc. All Rights Reserved.
+
+#include "OnlineBeaconUnitTestUtils.h"
+
+#include "Algo/Find.h"
+#include "Algo/RemoveIf.h"
+#include "Engine/Engine.h"
+#include "Engine/NetConnection.h"
+#include "Engine/World.h"
+#include "GameFramework/OnlineReplStructs.h"
+#include "IpNetDriver.h"
+#include "Misc/ConfigCacheIni.h"
+#include "Online/CoreOnline.h"
+#include "OnlineBeaconClient.h"
+#include "OnlineBeaconHost.h"
+#include "OnlineBeaconHostObject.h"
+#include "OnlineBeaconUnitTestNetDriver.h"
+#include "OnlineBeaconUnitTestSocketSubsystem.h"
+#include "PacketHandler.h"
+#include "Tests/AutomationEditorCommon.h"
+#include "TimerManager.h"
+
+namespace BeaconUnitTest
+{
+namespace Private
+{
+
+class FNetworkStatsImpl : public FNetworkStats, public FNetworkNotify
+{
+public:
+	FNetworkStatsImpl(const TObjectPtr<UNetDriver>& NetDriver)
+		: WeakNetDriver(NetDriver)
+		, PreviousNotifyHandler(nullptr)
+	{
+		if (NetDriver)
+		{
+			PreviousNotifyHandler = NetDriver->Notify;
+			NetDriver->Notify = this;
+		}
+	}
+
+	virtual ~FNetworkStatsImpl()
+	{
+		if (UNetDriver* NetDriver = WeakNetDriver.Get())
+		{
+			if (NetDriver->Notify == this)
+			{
+				NetDriver->Notify = PreviousNotifyHandler;
+			}
+			else
+			{
+				NetDriver->Notify = nullptr;
+			}
+		}
+	}
+
+	virtual EAcceptConnection::Type NotifyAcceptingConnection()
+	{
+		return PreviousNotifyHandler ? PreviousNotifyHandler->NotifyAcceptingConnection() : EAcceptConnection::Ignore;
+	}
+
+	virtual void NotifyAcceptedConnection(class UNetConnection* Connection)
+	{
+		if (PreviousNotifyHandler)
+		{
+			PreviousNotifyHandler->NotifyAcceptedConnection(Connection);
+		}
+	}
+
+	virtual bool NotifyAcceptingChannel(class UChannel* Channel)
+	{
+		return PreviousNotifyHandler ? PreviousNotifyHandler->NotifyAcceptingChannel(Channel) : false;
+	}
+
+	virtual void NotifyControlMessage(UNetConnection* Connection, uint8 MessageType, class FInBunch& Bunch)
+	{
+		ReceivedControlMessages.Add(MessageType);
+
+		if (PreviousNotifyHandler)
+		{
+			PreviousNotifyHandler->NotifyControlMessage(Connection, MessageType, Bunch);
+		}
+	}
+
+	TWeakObjectPtr<UNetDriver> WeakNetDriver;
+	FNetworkNotify* PreviousNotifyHandler = nullptr;
+};
+
+class FTestPrerequisitesImpl final : public FTestPrerequisites
+{
+public:
+	static TSharedPtr<FTestPrerequisitesImpl> TryCreate()
+	{
+		// Verify that the global net delegates are not in use before using them for the test.
+		if (!ensure(!FNetDelegates::OnReceivedNetworkEncryptionToken.IsBound()))
+		{
+			return nullptr;
+		}
+		if (!ensure(!FNetDelegates::OnReceivedNetworkEncryptionAck.IsBound()))
+		{
+			return nullptr;
+		}
+		if (!ensure(!FNetDelegates::OnReceivedNetworkEncryptionFailure.IsBound()))
+		{
+			return nullptr;
+		}
+
+		if (!ensure(GConfig))
+		{
+			return nullptr;
+		}
+
+		if (!ensure(GEngine))
+		{
+			return nullptr;
+		}
+
+		TSharedRef<FTestPrerequisitesImpl> Prerequisites = MakeShared<FTestPrerequisitesImpl>();
+		ActiveTestPrerequisites = Prerequisites;
+		return Prerequisites;
+	}
+
+	static TSharedPtr<FTestPrerequisitesImpl> Get()
+	{
+		return ActiveTestPrerequisites.Pin();
+	}
+
+	static FTestStats* GetActiveTestStats()
+	{
+		TSharedPtr<FTestPrerequisitesImpl> Prerequisites = ActiveTestPrerequisites.Pin();
+		return Prerequisites ? &Prerequisites->GetStats() : nullptr;
+	}
+
+	static const FTestConfig* GetActiveTestConfig()
+	{
+		TSharedPtr<FTestPrerequisitesImpl> Prerequisites = ActiveTestPrerequisites.Pin();
+		return Prerequisites ? &Prerequisites->GetConfig() : nullptr;
+	}
+
+	FTestPrerequisitesImpl()
+		: FTestPrerequisites()
+	{
+		SocketSubsystem = reinterpret_cast<FOnlineBeaconUnitTestSocketSubsystem*>(ISocketSubsystem::Get(SocketSubsystemName));
+
+		// Initialize world.
+		World = FAutomationEditorCommonUtils::CreateNewMap();
+		World->InitializeActorsForPlay(FURL());
+
+		// Bind delegates for handling network encryption.
+		BindNetEncryptionDelegates();
+
+		// Bind delegate to flush all outbound network messages to their target connections at the end of world tick.
+		NetworkFlushDelegateHandle = FWorldDelegates::OnWorldTickEnd.AddRaw(this, &FTestPrerequisitesImpl::NetworkFlush);
+
+#if !NO_LOGGING
+		// Disable logging for some categories since the tests will cause warnings / errors to be logged for some failure testing.
+		StoredLogBeaconVerbosity = LogBeacon.GetVerbosity();
+		LogBeacon.SetVerbosity(ELogVerbosity::NoLogging);
+		StoredLogNetVerbosity = LogNet.GetVerbosity();
+		LogNet.SetVerbosity(ELogVerbosity::NoLogging);
+#endif
+
+		// Setup encrpytion component. Store previous value to restore when the test has completed.
+		GConfig->GetString(TEXT("PacketHandlerComponents"), TEXT("EncryptionComponent"), StoredEncryptionComponentName, GEngineIni);
+		GConfig->SetString(TEXT("PacketHandlerComponents"), TEXT("EncryptionComponent"), *Config.Encryption.NetDriverEncryptionComponentName, GEngineIni);
+
+		// Install netdriver definition for the test driver.
+		const bool bDefinitionExists = Algo::FindBy(GEngine->NetDriverDefinitions, BeaconUnitTest::NetDriverDefinitionName, [](const FNetDriverDefinition& Definition)->const FName& { return Definition.DefName; }) != nullptr;
+		if (ensure(!bDefinitionExists))
+		{
+			FNetDriverDefinition TestBeaconNetdriverDefinition;
+			TestBeaconNetdriverDefinition.DefName = BeaconUnitTest::NetDriverDefinitionName;
+			TestBeaconNetdriverDefinition.DriverClassName = FName(UOnlineBeaconUnitTestNetDriver::StaticClass()->GetPathName());
+			GEngine->NetDriverDefinitions.Add(TestBeaconNetdriverDefinition);
+		}
+	}
+
+	virtual ~FTestPrerequisitesImpl()
+	{
+		// Remove unit test net driver definition.
+		GEngine->NetDriverDefinitions.SetNum(Algo::RemoveIf(GEngine->NetDriverDefinitions, [&](const FNetDriverDefinition& Definition) -> bool { return Definition.DefName == BeaconUnitTest::NetDriverDefinitionName; }));
+
+		FWorldDelegates::OnWorldTickEnd.Remove(NetworkFlushDelegateHandle);
+		UnbindNetEncryptionDelegates();
+		World = nullptr;
+
+#if !NO_LOGGING
+		// Restore logging.
+		LogBeacon.SetVerbosity(StoredLogBeaconVerbosity);
+		LogNet.SetVerbosity(StoredLogNetVerbosity);
+#endif
+
+		// Restore encryption config.
+		GConfig->SetString(TEXT("PacketHandlerComponents"), TEXT("EncryptionComponent"), *StoredEncryptionComponentName, GEngineIni);
+	}
+
+	virtual void BindNetEncryptionDelegates() override
+	{
+		FNetDelegates::OnReceivedNetworkEncryptionToken.BindRaw(this, &FTestPrerequisitesImpl::ReceivedNetworkEncryptionToken);
+		FNetDelegates::OnReceivedNetworkEncryptionAck.BindRaw(this, &FTestPrerequisitesImpl::ReceivedNetworkEncryptionAck);
+		FNetDelegates::OnReceivedNetworkEncryptionFailure.BindRaw(this, &FTestPrerequisitesImpl::ReceivedNetworkEncryptionFailure);
+	}
+
+	virtual void UnbindNetEncryptionDelegates() override
+	{
+		FNetDelegates::OnReceivedNetworkEncryptionToken.Unbind();
+		FNetDelegates::OnReceivedNetworkEncryptionAck.Unbind();
+		FNetDelegates::OnReceivedNetworkEncryptionFailure.Unbind();
+	}
+
+	void NetworkFlush(UWorld*, ELevelTick, float)
+	{
+		SocketSubsystem->FlushSendBuffers();
+	}
+
+	void ReceivedNetworkEncryptionToken(const FString& EncryptionToken, const FOnEncryptionKeyResponse& Delegate)
+	{
+		++Stats.Encryption.NetworkEncryptionToken.InvokeCount;
+
+		FEncryptionKeyResponse Response;
+		Response.Response = Config.Encryption.Host.Response;
+		Response.ErrorMsg = Config.Encryption.Host.ErrorMsg;
+		Response.EncryptionData = Config.Encryption.Host.EncryptionData;
+
+		if (Config.Encryption.Host.bDelayDelegate)
+		{
+			SetTimerForNextFrame(World, GFrameCounter, [this, Delegate, Response]()
+			{
+				++Stats.Encryption.NetworkEncryptionToken.CallbackCount;
+				Delegate.ExecuteIfBound(Response);
+			});
+		}
+		else
+		{
+			++Stats.Encryption.NetworkEncryptionToken.CallbackCount;
+			Delegate.ExecuteIfBound(Response);
+		}
+	}
+
+	void ReceivedNetworkEncryptionAck(const FOnEncryptionKeyResponse& Delegate)
+	{
+		++Stats.Encryption.NetworkEncryptionAck.InvokeCount;
+
+		FEncryptionKeyResponse Response;
+		Response.Response = Config.Encryption.Client.Response;
+		Response.ErrorMsg = Config.Encryption.Client.ErrorMsg;
+		Response.EncryptionData = Config.Encryption.Client.EncryptionData;
+
+		if (Config.Encryption.Client.bDelayDelegate)
+		{
+			SetTimerForNextFrame(World, GFrameCounter, [this, Delegate, Response]()
+			{
+				++Stats.Encryption.NetworkEncryptionAck.CallbackCount;
+				Delegate.ExecuteIfBound(Response);
+			});
+		}
+		else
+		{
+			++Stats.Encryption.NetworkEncryptionAck.CallbackCount;
+			Delegate.ExecuteIfBound(Response);
+		}
+	}
+
+	EEncryptionFailureAction ReceivedNetworkEncryptionFailure(UNetConnection* Connection)
+	{
+		++Stats.Encryption.NetworkEncryptionFailure.InvokeCount;
+		return Config.Encryption.FailureAction;
+	}
+
+	FOnlineBeaconUnitTestSocketSubsystem* SocketSubsystem = nullptr;
+	FDelegateHandle NetworkFlushDelegateHandle;
+
+#if !NO_LOGGING
+	ELogVerbosity::Type StoredLogBeaconVerbosity = ELogVerbosity::NoLogging;
+	ELogVerbosity::Type StoredLogNetVerbosity = ELogVerbosity::NoLogging;
+#endif
+
+	FString StoredEncryptionComponentName;
+
+	static TWeakPtr<FTestPrerequisitesImpl> ActiveTestPrerequisites;
+};
+
+TWeakPtr<FTestPrerequisitesImpl> FTestPrerequisitesImpl::ActiveTestPrerequisites;
+
+//--------------------------------------------------------------------------------------------------
+// Hacky accessor to private fields - Allowed due to a loophole in the c++ standard regarding ignoring access checks for template specialization.
+
+// OnlineBeaconHost access.
+
+template<typename ClassType, typename ClassType::MemberType Member>
+struct TBeaconHostAccessShim
+{
+	friend typename ClassType::MemberType UnitTestGetBeaconHostMember(ClassType)
+	{
+		return Member;
+	}
+};
+struct FBeaconHostClientActorsTag
+{ 
+  typedef TArray<TObjectPtr<AOnlineBeaconClient>> AOnlineBeaconHostObject::*MemberType;
+  friend MemberType UnitTestGetBeaconHostMember(FBeaconHostClientActorsTag);
+};
+template struct TBeaconHostAccessShim<FBeaconHostClientActorsTag, &AOnlineBeaconHostObject::ClientActors>;
+
+struct FBeaconHostAuthRequiredTag
+{ 
+  typedef bool AOnlineBeaconHost::*MemberType;
+  friend MemberType UnitTestGetBeaconHostMember(FBeaconHostAuthRequiredTag);
+};
+template struct TBeaconHostAccessShim<FBeaconHostAuthRequiredTag, &AOnlineBeaconHost::bAuthRequired>;
+
+//--------------------------------------------------------------------------------------------------
+
+// OnlineBeacon access.
+
+template<typename ClassType, typename ClassType::MemberType Member>
+struct TBeaconAccessShim
+{
+	friend typename ClassType::MemberType UnitTestGetBeaconMember(ClassType)
+	{
+		return Member;
+	}
+};
+
+struct FBeaconNetDriverTag
+{ 
+  typedef TObjectPtr<UNetDriver> AOnlineBeacon::*MemberType;
+  friend MemberType UnitTestGetBeaconMember(FBeaconNetDriverTag);
+};
+template struct TBeaconAccessShim<FBeaconNetDriverTag, &AOnlineBeacon::NetDriver>;
+
+struct FBeaconNetDriverDefinitionNameTag
+{ 
+  typedef FName AOnlineBeacon::*MemberType;
+  friend MemberType UnitTestGetBeaconMember(FBeaconNetDriverDefinitionNameTag);
+};
+template struct TBeaconAccessShim<FBeaconNetDriverDefinitionNameTag, &AOnlineBeacon::NetDriverDefinitionName>;
+
+/* Private */}
+
+const FName NetDriverDefinitionName(TEXT("UnitTestBeaconNetDriver"));
+const FName SocketSubsystemName(TEXT("UnitTestSocketSubsystem"));
+
+FTestStats* FTestPrerequisites::GetActiveTestStats()
+{
+	return Private::FTestPrerequisitesImpl::GetActiveTestStats();
+}
+
+const FTestConfig* FTestPrerequisites::GetActiveTestConfig()
+{
+	return Private::FTestPrerequisitesImpl::GetActiveTestConfig();
+}
+
+TSharedPtr<FTestPrerequisites> FTestPrerequisites::TryCreate()
+{
+	return Private::FTestPrerequisitesImpl::TryCreate();
+}
+
+bool SetSocketFlags(AOnlineBeacon* OnlineBeacon, ESocketFlags Flags)
+{
+	if (OnlineBeacon == nullptr)
+	{
+		return false;
+	}
+
+	TObjectPtr<UNetDriver>& NetDriver = OnlineBeacon->*UnitTestGetBeaconMember(Private::FBeaconNetDriverTag());
+	UIpNetDriver* IpNetDriver = Cast<UIpNetDriver>(NetDriver.Get());
+	FSocketBeaconUnitTest* Socket = IpNetDriver ? reinterpret_cast<FSocketBeaconUnitTest*>(IpNetDriver->GetSocket()) : nullptr;
+
+	if (Socket == nullptr)
+	{
+		return false;
+	}
+
+	Socket->SetUnitTestFlags(Flags);
+	return true;
+}
+
+bool ConfigureBeacon(const FTestPrerequisites& Prerequisites, AOnlineBeacon* OnlineBeacon)
+{
+	if (OnlineBeacon == nullptr)
+	{
+		return false;
+	}
+
+	const BeaconUnitTest::FTestConfig& TestConfig = Prerequisites.GetConfig();
+
+	// Set beacon to use the unit test net driver.
+	FName& BeaconNetDriverDefinitionName = OnlineBeacon->*UnitTestGetBeaconMember(Private::FBeaconNetDriverDefinitionNameTag());
+	BeaconNetDriverDefinitionName = NetDriverDefinitionName;
+
+	if (AOnlineBeaconHost* BeaconHost = Cast<AOnlineBeaconHost>(OnlineBeacon))
+	{
+		bool& bAuthRequired = BeaconHost->*UnitTestGetBeaconHostMember(Private::FBeaconHostAuthRequiredTag());
+		bAuthRequired = TestConfig.Auth.bEnabled;
+	}
+
+	return true;
+}
+
+bool ConfigureBeaconNetDriver(const FTestPrerequisites& Prerequisites, AOnlineBeacon* OnlineBeacon, TSharedPtr<FNetworkStats>* OutStats)
+{
+	if (OnlineBeacon == nullptr)
+	{
+		return false;
+	}
+
+	TObjectPtr<UNetDriver>& NetDriver = OnlineBeacon->*UnitTestGetBeaconMember(Private::FBeaconNetDriverTag());
+	if (NetDriver == nullptr)
+	{
+		return false;
+	}
+
+	const BeaconUnitTest::FTestConfig TestConfig = Prerequisites.GetConfig();
+	NetDriver->bNoTimeouts = TestConfig.NetDriver.bNoTimeouts;
+	NetDriver->InitialConnectTimeout = TestConfig.NetDriver.InitialConnectTimeout;
+	NetDriver->ConnectionTimeout = TestConfig.NetDriver.ConnectionTimeout;
+	NetDriver->KeepAliveTime = TestConfig.NetDriver.ConnectionTimeout;
+
+	if (OutStats)
+	{
+		*OutStats = MakeShared<Private::FNetworkStatsImpl>(NetDriver);
+	}
+
+	return true;
+}
+
+bool InitClientForUser(const FTestPrerequisites& Prerequisites, AOnlineBeaconClient* OnlineBeaconClient, const FUniqueNetIdRef& User)
+{
+	if (OnlineBeaconClient == nullptr)
+	{
+		return false;
+	}
+
+	FURL URL;
+	URL.Port = Prerequisites.GetConfig().NetDriver.ServerListenPort;
+	if (!OnlineBeaconClient->InitClient(URL))
+	{
+		return false;
+	}
+
+	UNetConnection* Connection = OnlineBeaconClient->GetNetConnection();
+	if(Connection == nullptr)
+	{
+		return false;
+	}
+
+	Connection->PlayerId = User;
+	return true;
+}
+
+AOnlineBeaconClient* GetBeaconClientForUser(AOnlineBeaconHostObject* OnlineBeaconHostObject, const FUniqueNetIdRef& User)
+{
+	if (OnlineBeaconHostObject == nullptr)
+	{
+		return nullptr;
+	}
+
+	if (!User->IsValid())
+	{
+		return nullptr;
+	}
+
+	TArray<TObjectPtr<AOnlineBeaconClient>>& ClientActors = OnlineBeaconHostObject->*UnitTestGetBeaconHostMember(Private::FBeaconHostClientActorsTag());
+	TObjectPtr<AOnlineBeaconClient>* Result = Algo::FindByPredicate(ClientActors, [&User](const TObjectPtr<AOnlineBeaconClient>& BeaconClient){
+		return BeaconClient->GetUniqueId() == User;
+	});
+
+	return Result ? Result->Get() : nullptr;
+}
+
+bool TickOnce(const FTestPrerequisites& Prerequisites, ETickFlags Flags)
+{
+	UWorld* World = Prerequisites.GetWorld();
+	if (World == nullptr)
+	{
+		return false;
+	}
+
+	if (EnumHasAllFlags(Flags, ETickFlags::SleepTickTime))
+	{
+		FPlatformProcess::Sleep(Prerequisites.GetConfig().WorldTickRate);
+	}
+	World->Tick(ELevelTick::LEVELTICK_All, Prerequisites.GetConfig().WorldTickRate);
+	++GFrameCounter;
+
+	return true;
+}
+
+bool TickUntilConnectionInitialized(const FTestPrerequisites& Prerequisites, AOnlineBeaconClient* OnlineBeaconClient, ETickFlags Flags)
+{
+	if (OnlineBeaconClient == nullptr)
+	{
+		return false;
+	}
+
+	TWeakObjectPtr<AOnlineBeaconClient> WeakOnlineBeaconClient = OnlineBeaconClient;
+	return TickUntil(Prerequisites, [WeakOnlineBeaconClient]()
+	{
+		AOnlineBeaconClient* OnlineBeaconClient = WeakOnlineBeaconClient.Get();
+		UNetConnection* BeaconConnection = OnlineBeaconClient ? OnlineBeaconClient->GetNetConnection() : nullptr;
+		return BeaconConnection && BeaconConnection->Handler ? BeaconConnection->Handler->IsFullyInitialized() : false;
+	}, Flags);
+}
+
+bool TickUntilControlMessageReceived(const FTestPrerequisites& Prerequisites, const FNetworkStats& Stats, ETickFlags Flags)
+{
+	return TickUntil(Prerequisites, [&Stats]() -> bool
+	{
+		return !Stats.ReceivedControlMessages.IsEmpty();
+	});
+}
+
+bool TickUntilConnected(const FTestPrerequisites& Prerequisites, AOnlineBeaconClient* OnlineBeaconClient, ETickFlags Flags)
+{
+	TWeakObjectPtr<AOnlineBeaconClient> WeakOnlineBeaconClient = OnlineBeaconClient;
+	return TickUntil(Prerequisites, [WeakOnlineBeaconClient]()
+	{
+		AOnlineBeaconClient* OnlineBeaconClient = WeakOnlineBeaconClient.Get();
+		return OnlineBeaconClient != nullptr && (OnlineBeaconClient->GetConnectionState() == EBeaconConnectionState::Open);
+	}, Flags);
+}
+
+bool TickUntilDisconnected(const FTestPrerequisites& Prerequisites, AOnlineBeaconClient* OnlineBeaconClient, ETickFlags Flags)
+{
+	TWeakObjectPtr<AOnlineBeaconClient> WeakOnlineBeaconClient = OnlineBeaconClient;
+	return TickUntil(Prerequisites, [WeakOnlineBeaconClient]()
+	{
+		AOnlineBeaconClient* OnlineBeaconClient = WeakOnlineBeaconClient.Get();
+		return OnlineBeaconClient != nullptr && (OnlineBeaconClient->GetConnectionState() == EBeaconConnectionState::Invalid || OnlineBeaconClient->GetConnectionState() == EBeaconConnectionState::Closed);
+	}, Flags);
+}
+
+bool TickUntilTimeoutElapsed(const FTestPrerequisites& Prerequisite, ETickFlags Flags)
+{
+	const double StartTime = FPlatformTime::Seconds();
+	const double ExpireTime = StartTime + Prerequisite.GetConfig().NetDriver.InitialConnectTimeout;
+
+	// Wait for time to expire.
+	return BeaconUnitTest::TickUntil(Prerequisite, [ExpireTime]() -> bool
+	{
+		return FPlatformTime::Seconds() > ExpireTime;
+	}, BeaconUnitTest::ETickFlags::SleepTickTime);
+}
+
+/* EBeaconUnitTest */ }
diff --git a/Engine/Plugins/Online/OnlineSubsystemUtils/Source/OnlineSubsystemUtilsTest/Private/OnlineSubsystemUtilsTestModule.cpp b/Engine/Plugins/Online/OnlineSubsystemUtils/Source/OnlineSubsystemUtilsTest/Private/OnlineSubsystemUtilsTestModule.cpp
new file mode 100644
--- /dev/null
+++ b/Engine/Plugins/Online/OnlineSubsystemUtils/Source/OnlineSubsystemUtilsTest/Private/OnlineSubsystemUtilsTestModule.cpp
@@ -0,0 +1,45 @@
+// Copyright Epic Games, Inc. All Rights Reserved.
+
+#include "Modules/ModuleInterface.h"
+#include "Modules/ModuleManager.h"
+#include "OnlineBeaconUnitTestSocketSubsystem.h"
+
+/**
+ * Online subsystem utils test module class
+ * Test module used to implement automation testing to be handled within editor builds.
+ */
+class FOnlineSubsystemUtilsTestModule : public IModuleInterface
+{
+public:
+
+	FOnlineSubsystemUtilsTestModule()
+	{
+	}
+
+	virtual ~FOnlineSubsystemUtilsTestModule()
+	{
+	}
+
+	// IModuleInterface
+	virtual void StartupModule() override
+	{
+		FString InitError;
+		SocketSubsystem = MakeShared<FOnlineBeaconUnitTestSocketSubsystem>();
+		ensure(SocketSubsystem->Init(InitError));
+	}
+
+	virtual void ShutdownModule() override
+	{
+		SocketSubsystem->Shutdown();
+		SocketSubsystem.Reset();
+	}
+
+	virtual bool SupportsDynamicReloading() override
+	{
+		return true;
+	}
+
+	TSharedPtr<FOnlineBeaconUnitTestSocketSubsystem> SocketSubsystem;
+};
+
+IMPLEMENT_MODULE(FOnlineSubsystemUtilsTestModule, OnlineSubsystemUtilsTest);
diff --git a/Engine/Plugins/Online/OnlineSubsystemUtils/Source/OnlineSubsystemUtilsTest/Public/OnlineBeaconUnitTestUtils.h b/Engine/Plugins/Online/OnlineSubsystemUtils/Source/OnlineSubsystemUtilsTest/Public/OnlineBeaconUnitTestUtils.h
new file mode 100644
--- /dev/null
+++ b/Engine/Plugins/Online/OnlineSubsystemUtils/Source/OnlineSubsystemUtilsTest/Public/OnlineBeaconUnitTestUtils.h
@@ -0,0 +1,382 @@
+// Copyright Epic Games, Inc. All Rights Reserved.
+
+#pragma once
+
+#include "CoreMinimal.h"
+
+#include "Engine/NetworkDelegates.h"
+#include "Engine/World.h"
+#include "Online/CoreOnlineFwd.h"
+#include "OnlineError.h"
+#include "Templates/Function.h"
+#include "TimerManager.h"
+#include "UObject/ObjectPtr.h"
+
+class AOnlineBeacon;
+class AOnlineBeaconClient;
+class AOnlineBeaconHostObject;
+class FNetworkNotify;
+class ISocketSubsystem;
+struct FURL;
+class UWorld;
+
+namespace BeaconUnitTest
+{
+	extern const FName NetDriverDefinitionName;
+	extern const FName SocketSubsystemName;
+
+	enum class ESocketFlags : uint8
+	{
+		Disabled = 0,
+		SendEnabled = 1 << 0,
+		RecvEnabled = 1 << 1,
+
+		Default = SendEnabled | RecvEnabled,
+	};
+	ENUM_CLASS_FLAGS(ESocketFlags);
+
+	enum class ETickFlags : uint8
+	{
+		None = 0,
+
+		/**
+		 * When set, sleep during the tick duration to progress wall time.
+		 * Network timeouts depend on FPlatformTime::Seconds() instead of world tick time.
+		 */
+		SleepTickTime = 1 << 0,
+	};
+	ENUM_CLASS_FLAGS(ETickFlags);
+
+	class FNetworkStats
+	{
+	public:
+		TArray<uint8> ReceivedControlMessages;
+	};
+
+	struct FTestStats
+	{
+		struct FEncryption
+		{
+			struct FNetworkEncryptionToken
+			{
+				uint32 InvokeCount = 0;
+				uint32 CallbackCount = 0;
+			};
+
+			struct FNetworkEncryptionAck
+			{
+				uint32 InvokeCount = 0;
+				uint32 CallbackCount = 0;
+			};
+
+			struct FNetworkEncryptionFailure
+			{
+				uint32 InvokeCount = 0;
+			};
+
+			FNetworkEncryptionToken NetworkEncryptionToken;
+			FNetworkEncryptionAck NetworkEncryptionAck;
+			FNetworkEncryptionFailure NetworkEncryptionFailure;
+		};
+
+		struct FClient
+		{
+			struct FOnConnected
+			{
+				uint32 InvokeCount = 0;
+			};
+
+			struct FOnFailure
+			{
+				uint32 InvokeCount = 0;
+			};
+
+			FOnConnected OnConnected;
+			FOnFailure OnFailure;
+		};
+
+		struct FHost
+		{
+			struct FOnFailure
+			{
+				uint32 InvokeCount = 0;
+			};
+
+			FOnFailure OnFailure;
+		};
+
+		struct FHostObject
+		{
+			struct FOnClientConnected
+			{
+				uint32 InvokeCount = 0;
+			};
+
+			struct FNotifyClientDisconnected
+			{
+				uint32 InvokeCount = 0;
+			};
+
+			FOnClientConnected OnClientConnected;
+			FNotifyClientDisconnected NotifyClientDisconnected;
+		};
+
+		FEncryption Encryption;
+		FClient Client;
+		FHost Host;
+		FHostObject HostObject;
+	};
+
+	struct FTestConfig
+	{
+		struct FNetDriver
+		{
+			bool bNoTimeouts = false;
+			bool bFailInit = false;
+
+			float ConnectionTimeout = 1/15.f;
+
+			float InitialConnectTimeout = 1/30.f;
+
+			/* 
+			 * Keepalive packets will also flush any waiting outbound messages.
+			 * 
+			 * This should normally be set to a time more frequent than the connection timeout, but
+			 * less frequent than the world tick rate. Beacons should manually flush after each
+			 * message to prevent needing to wait on a keep-alive flush to progress the handshake.
+			 */
+			float KeepAliveTime =  1/60.f;
+
+			int32 ServerListenPort = 0;
+		};
+
+		struct FEncryption
+		{
+			struct FEncryptionHost
+			{
+				bool bDelayDelegate = false;
+				EEncryptionResponse Response;
+				FString ErrorMsg;
+				FEncryptionData EncryptionData;
+			};
+
+			struct FEncryptionClient
+			{
+				bool bDelayDelegate = false;
+				EEncryptionResponse Response;
+				FString ErrorMsg;
+				FEncryptionData EncryptionData;
+			};
+
+			bool bEnabled = false;
+
+			FEncryptionHost Host;
+			FEncryptionClient Client;
+			EEncryptionFailureAction FailureAction = EEncryptionFailureAction::Default;
+			FString NetDriverEncryptionComponentName;
+		};
+
+		struct FAuth
+		{
+			struct FMethod1
+			{
+				bool bEnabled = false;
+				bool bDelayDelegate = false;
+				FOnlineError Result;
+			};
+
+			struct FMethod2
+			{
+				bool bEnabled = false;
+				bool bDelayDelegate = false;
+				FOnlineError Result;
+			};
+
+			struct FMethod3
+			{
+				bool bEnabled = false;
+				bool bDelayDelegate = false;
+				FOnlineError Result;
+			};
+
+			struct FVerify
+			{
+				bool bEnabled = false;
+				bool bResult = false;
+			};
+
+			bool bEnabled = false;
+			FMethod1 Method1;
+			FMethod2 Method2;
+			FMethod3 Method3;
+			FVerify Verify;
+		};
+
+		struct FClient
+		{
+			struct FOnConnected
+			{
+				TFunction<void()> Callback;
+			};
+
+			FOnConnected OnConnected;
+		};
+
+		float WorldTickRate = 1/120.f;
+
+		FNetDriver NetDriver;
+		FEncryption Encryption;
+		FAuth Auth;
+		FClient Client;
+	};
+
+	class FTestPrerequisites : public FNoncopyable, public TSharedFromThis<FTestPrerequisites>
+	{
+	public:
+		virtual ~FTestPrerequisites() = default;
+
+		static ONLINESUBSYSTEMUTILSTEST_API TSharedPtr<FTestPrerequisites> TryCreate();
+		static ONLINESUBSYSTEMUTILSTEST_API TSharedPtr<FTestPrerequisites> Get();
+
+		static ONLINESUBSYSTEMUTILSTEST_API FTestStats* GetActiveTestStats();
+		static ONLINESUBSYSTEMUTILSTEST_API const FTestConfig* GetActiveTestConfig();
+
+		virtual void BindNetEncryptionDelegates() = 0;
+		virtual void UnbindNetEncryptionDelegates() = 0;
+
+		UWorld* GetWorld() const
+		{
+			return World.Get();
+		}
+
+		FTestStats& GetStats()
+		{
+			return Stats;
+		}
+
+		const FTestStats& GetStats() const
+		{
+			return Stats;
+		}
+
+		FTestConfig& GetConfig()
+		{
+			return Config;
+		}
+
+		const FTestConfig& GetConfig() const
+		{
+			return Config;
+		}
+
+	protected:
+		FTestPrerequisites() = default;
+
+		TObjectPtr<UWorld> World;
+		FTestStats Stats;
+		FTestConfig Config;
+	};
+
+	/**
+	 * Adjust the socket flags for a beacon unit test socket associated with an online beacon.
+	 * When send or receive are disabled pending messages will be queued. The messages will be dispatched on the following tick after re-enabling the channel.
+	 */
+	ONLINESUBSYSTEMUTILSTEST_API bool SetSocketFlags(AOnlineBeacon* OnlineBeacon, ESocketFlags Flags);
+
+	ONLINESUBSYSTEMUTILSTEST_API bool ConfigureBeacon(const FTestPrerequisites& Prerequisites, AOnlineBeacon* OnlineBeacon);
+
+	ONLINESUBSYSTEMUTILSTEST_API bool ConfigureBeaconNetDriver(const FTestPrerequisites& Prerequisites, AOnlineBeacon* OnlineBeacon, TSharedPtr<FNetworkStats>* OutStats);
+
+	/**
+	 * Helper to initialize the connection for a beacon client and overriding the default UserId.
+	 */
+	ONLINESUBSYSTEMUTILSTEST_API bool InitClientForUser(const FTestPrerequisites& Prerequisites, AOnlineBeaconClient* OnlineBeaconClient, const FUniqueNetIdRef& User);
+
+	/**
+	 * Helper to get the hosts version of the client beacon actor for the specified user.
+	 */
+	ONLINESUBSYSTEMUTILSTEST_API AOnlineBeaconClient* GetBeaconClientForUser(AOnlineBeaconHostObject* HostObject, const FUniqueNetIdRef& User);
+
+
+	/**
+	 * Tick the unit test world one time to progress network state.
+	 */
+	ONLINESUBSYSTEMUTILSTEST_API bool TickOnce(const FTestPrerequisites& Prerequisites, ETickFlags Flags = ETickFlags::None);
+
+	/**
+	 * Tick the unit test world until the beacon client connection has completed the packet handler handshake.
+	 */
+	ONLINESUBSYSTEMUTILSTEST_API bool TickUntilConnectionInitialized(const FTestPrerequisites& Prerequisites, AOnlineBeaconClient* OnlineBeaconClient, ETickFlags Flags = ETickFlags::None);
+
+	/**
+	 * Tick the unit test world until the beacon client has changed its status to 'Closed' or 'Invalid'.
+	 */
+	ONLINESUBSYSTEMUTILSTEST_API bool TickUntilControlMessageReceived(const FTestPrerequisites& Prerequisites, const FNetworkStats& Stats, ETickFlags Flags = ETickFlags::None);
+
+	/**
+	 * Tick the unit test world until the beacon client has changed its status to 'Open'.
+	 */
+	ONLINESUBSYSTEMUTILSTEST_API bool TickUntilConnected(const FTestPrerequisites& Prerequisites, AOnlineBeaconClient* OnlineBeaconClient, ETickFlags Flags = ETickFlags::None);
+
+	/**
+	 * Tick the unit test world until the beacon client has changed its status to 'Closed' or 'Invalid'.
+	 */
+	ONLINESUBSYSTEMUTILSTEST_API bool TickUntilDisconnected(const FTestPrerequisites& Prerequisites, AOnlineBeaconClient* OnlineBeaconClient, ETickFlags Flags = ETickFlags::None);
+
+	/**
+	 * Tick the unit test world until the network timeout has elapsed.
+	 */
+	ONLINESUBSYSTEMUTILSTEST_API bool TickUntilTimeoutElapsed(const FTestPrerequisites& Prerequisite, ETickFlags Flags = ETickFlags::None);
+
+	/**
+	 * Tick the unit test world until predicated condition is met or max iterations is reached.
+	 */
+	template <typename FPredicate>
+	inline bool TickUntil(const FTestPrerequisites& Prerequisites, FPredicate Predicate, ETickFlags Flags = ETickFlags::None)
+	{
+		constexpr uint32 MaxIterations = 64;
+		for (uint32 i = 0; i < MaxIterations; ++i)
+		{
+			if (Predicate())
+			{
+				return true;
+			}
+
+			if (!TickOnce(Prerequisites, Flags))
+			{
+				return false;
+			}
+		}
+
+		return false;
+	}
+
+	/**
+	 * Schedule the callback to be triggered on the next frame. SetTimerForNextTick may run the
+	 * timer on the same frame as scheduled if the timer manager has not yet been ticked.
+	 */
+	template <typename TCallback>
+	bool SetTimerForNextFrame(UWorld* World, uint64 RequestingFrame, TCallback Callback)
+	{
+		if (ensure(World))
+		{
+			World->GetTimerManager().SetTimerForNextTick([World, RequestingFrame, Callback]()
+			{
+				if (RequestingFrame < GFrameCounter)
+				{
+					Callback();
+				}
+				else
+				{
+					// callback was signaled too early.
+					SetTimerForNextFrame(World, RequestingFrame, Callback);
+				}
+			});
+
+			return true;
+		}
+
+		return false;
+	}
+
+/* BeaconUnitTest */ }
diff --git a/Engine/Source/Runtime/Engine/Private/DataChannel.cpp b/Engine/Source/Runtime/Engine/Private/DataChannel.cpp
--- a/Engine/Source/Runtime/Engine/Private/DataChannel.cpp
+++ b/Engine/Source/Runtime/Engine/Private/DataChannel.cpp
@@ -1621,6 +1621,7 @@ void UChannel::AddedToChannelPool()
 -----------------------------------------------------------------------------*/
 
 const TCHAR* FNetControlMessageInfo::Names[256];
+bool FNetControlMessageInfo::bInitialized = false;
 
 // control channel message implementation
 IMPLEMENT_CONTROL_CHANNEL_MESSAGE(Hello);
diff --git a/Engine/Source/Runtime/Engine/Public/Net/DataChannel.h b/Engine/Source/Runtime/Engine/Public/Net/DataChannel.h
--- a/Engine/Source/Runtime/Engine/Public/Net/DataChannel.h
+++ b/Engine/Source/Runtime/Engine/Public/Net/DataChannel.h
@@ -51,7 +51,6 @@ public:
 private:
 	static void CheckInitialized()
 	{
-		static bool bInitialized = false;
 		if (!bInitialized)
 		{
 			for (int32 i = 0; i < UE_ARRAY_COUNT(Names); i++)
@@ -89,6 +88,7 @@ private:
 	
 	static constexpr int MaxNames = TNumericLimits<uint8>::Max() + 1;
 	static ENGINE_API const TCHAR* Names[MaxNames];
+	static ENGINE_API bool bInitialized;
 };
 
 #define DEFINE_CONTROL_CHANNEL_MESSAGE(Name, Index, ...) \
