Change init ordering for UOnlineBeaconUnitTestNetDriver so that Super::InitBase is always called.

#rnx

diff --git a/Engine/Plugins/Online/OnlineSubsystemUtils/Source/OnlineSubsystemUtilsTest/Private/OnlineBeaconUnitTestNetDriver.cpp b/Engine/Plugins/Online/OnlineSubsystemUtils/Source/OnlineSubsystemUtilsTest/Private/OnlineBeaconUnitTestNetDriver.cpp
--- a/Engine/Plugins/Online/OnlineSubsystemUtils/Source/OnlineSubsystemUtilsTest/Private/OnlineBeaconUnitTestNetDriver.cpp
+++ b/Engine/Plugins/Online/OnlineSubsystemUtils/Source/OnlineSubsystemUtilsTest/Private/OnlineBeaconUnitTestNetDriver.cpp
@@ -21,13 +21,16 @@ void UOnlineBeaconUnitTestNetDriver::PostInitProperties()
 
 bool UOnlineBeaconUnitTestNetDriver::InitBase(bool bInitAsClient, FNetworkNotify* InNotify, const FURL& URL, bool bReuseAddressAndPort, FString& Error)
 {
-	const BeaconUnitTest::FTestConfig* TestConfig = BeaconUnitTest::FTestPrerequisites::GetActiveTestConfig();
-	if (TestConfig == nullptr || TestConfig->NetDriver.bFailInit)
+	if (Super::InitBase(bInitAsClient, InNotify, URL, bReuseAddressAndPort, Error))
 	{
-		return false;
+		const BeaconUnitTest::FTestConfig* TestConfig = BeaconUnitTest::FTestPrerequisites::GetActiveTestConfig();
+		if (TestConfig == nullptr || TestConfig->NetDriver.bFailInit)
+		{
+			return false;
+		}
 	}
 
-	return Super::InitBase(bInitAsClient, InNotify, URL, bReuseAddressAndPort, Error);
+	return true;
 }
 
 void UOnlineBeaconUnitTestNetDriver::TickDispatch(float DeltaTime)
