Fix build issue

#rnx

diff --git a/Engine/Plugins/Online/OnlineSubsystemUtils/Source/OnlineSubsystemUtilsTest/Private/OnlineBeaconUnitTestSocketSubsystem.cpp b/Engine/Plugins/Online/OnlineSubsystemUtils/Source/OnlineSubsystemUtilsTest/Private/OnlineBeaconUnitTestSocketSubsystem.cpp
--- a/Engine/Plugins/Online/OnlineSubsystemUtils/Source/OnlineSubsystemUtilsTest/Private/OnlineBeaconUnitTestSocketSubsystem.cpp
+++ b/Engine/Plugins/Online/OnlineSubsystemUtils/Source/OnlineSubsystemUtilsTest/Private/OnlineBeaconUnitTestSocketSubsystem.cpp
@@ -498,7 +498,7 @@ FSocket* FOnlineBeaconUnitTestSocketSubsystem::CreateSocket(const FName& SocketT
 		return nullptr;
 	}
 
-	FSocketInfo& SocketInfo = Sockets.Emplace_GetRef(MakeShared<FSocketBeaconUnitTest>(SOCKTYPE_Datagram, SocketDescription, ProtocolType, AsShared()), FInternetAddrBeaconUnitTest(), false);
+	FSocketInfo& SocketInfo = Sockets.Emplace_GetRef(FSocketInfo{MakeShared<FSocketBeaconUnitTest>(SOCKTYPE_Datagram, SocketDescription, ProtocolType, AsShared()), FInternetAddrBeaconUnitTest(), false});
 	return SocketInfo.Socket.Get();
 }
 
