[Backout] - CL30794713
Original CL Desc
-----------------------------------------------------------------
Change init ordering for UOnlineBeaconUnitTestNetDriver so that Super::InitBase is always called.

#rnx

diff --git a/Engine/Plugins/Online/OnlineSubsystemUtils/Source/OnlineSubsystemUtilsTest/Private/OnlineBeaconUnitTestNetDriver.cpp b/Engine/Plugins/Online/OnlineSubsystemUtils/Source/OnlineSubsystemUtilsTest/Private/OnlineBeaconUnitTestNetDriver.cpp
--- a/Engine/Plugins/Online/OnlineSubsystemUtils/Source/OnlineSubsystemUtilsTest/Private/OnlineBeaconUnitTestNetDriver.cpp
+++ b/Engine/Plugins/Online/OnlineSubsystemUtils/Source/OnlineSubsystemUtilsTest/Private/OnlineBeaconUnitTestNetDriver.cpp
@@ -21,16 +21,13 @@ void UOnlineBeaconUnitTestNetDriver::PostInitProperties()
 
 bool UOnlineBeaconUnitTestNetDriver::InitBase(bool bInitAsClient, FNetworkNotify* InNotify, const FURL& URL, bool bReuseAddressAndPort, FString& Error)
 {
-	if (Super::InitBase(bInitAsClient, InNotify, URL, bReuseAddressAndPort, Error))
+	const BeaconUnitTest::FTestConfig* TestConfig = BeaconUnitTest::FTestPrerequisites::GetActiveTestConfig();
+	if (TestConfig == nullptr || TestConfig->NetDriver.bFailInit)
 	{
-		const BeaconUnitTest::FTestConfig* TestConfig = BeaconUnitTest::FTestPrerequisites::GetActiveTestConfig();
-		if (TestConfig == nullptr || TestConfig->NetDriver.bFailInit)
-		{
-			return false;
-		}
+		return false;
 	}
 
-	return true;
+	return Super::InitBase(bInitAsClient, InNotify, URL, bReuseAddressAndPort, Error);
 }
 
 void UOnlineBeaconUnitTestNetDriver::TickDispatch(float DeltaTime)
