[Backout] - CL30791987
Original CL Desc
-----------------------------------------------------------------
Add missing network flush after client sends NMT_BeaconNetGUIDAck.
Add missing network flush after host sends NMT_BeaconAssignGUID and ClientOnConnected RPC.
Fix FNetControlMessageInfo to handle intializing message type names in editor builds.
Add unit tests for OnlineBeacon.

diff --git a/Engine/Plugins/Online/OnlineSubsystemUtils/OnlineSubsystemUtils.uplugin b/Engine/Plugins/Online/OnlineSubsystemUtils/OnlineSubsystemUtils.uplugin
--- a/Engine/Plugins/Online/OnlineSubsystemUtils/OnlineSubsystemUtils.uplugin
+++ b/Engine/Plugins/Online/OnlineSubsystemUtils/OnlineSubsystemUtils.uplugin
@@ -1,33 +1,27 @@
 {
-	"FileVersion": 3,
+	"FileVersion" : 3,
+	
+	"FriendlyName" : "Online Subsystem Utils",
+	"Version" : 1,
+	"VersionName" : "1.0",
+	"Description" : "Shared code for interacting online service and online subsystem implementations.",
+	"Category" : "Online Platform",
+	"CreatedBy" : "Epic Games, Inc.",
+	"CreatedByURL" : "https://epicgames.com",
+	"EnabledByDefault" : true,
 
-	"FriendlyName": "Online Subsystem Utils",
-	"Version": 1,
-	"VersionName": "1.0",
-	"Description": "Shared code for interacting online service and online subsystem implementations.",
-	"Category": "Online Platform",
-	"CreatedBy": "Epic Games, Inc.",
-	"CreatedByURL": "https://epicgames.com",
-	"EnabledByDefault": true,
-
-	"Modules":
-	[
-		{
-			"Name": "OnlineSubsystemUtils",
-			"Type": "Runtime",
-			"LoadingPhase": "PostConfigInit"
-		},
-		{
-			"Name": "OnlineSubsystemUtilsTest",
-			"Type": "Editor",
-			"LoadingPhase": "PostEngineInit"
-		},
-		{
-			"Name": "OnlineBlueprintSupport",
-			"Type": "UncookedOnly",
-			"LoadingPhase": "PostConfigInit"
-		}
-	],
+    "Modules": [
+        {
+            "Name": "OnlineSubsystemUtils",
+            "Type": "Runtime",
+            "LoadingPhase": "PostConfigInit"
+        },
+        {
+            "Name": "OnlineBlueprintSupport",
+            "Type": "UncookedOnly",
+            "LoadingPhase": "PostConfigInit"
+        }
+    ],
 	"LocalizationTargets":
 	[
 		{
diff --git a/Engine/Plugins/Online/OnlineSubsystemUtils/Source/OnlineSubsystemUtils/Private/OnlineBeaconClient.cpp b/Engine/Plugins/Online/OnlineSubsystemUtils/Source/OnlineSubsystemUtils/Private/OnlineBeaconClient.cpp
--- a/Engine/Plugins/Online/OnlineSubsystemUtils/Source/OnlineSubsystemUtils/Private/OnlineBeaconClient.cpp
+++ b/Engine/Plugins/Online/OnlineSubsystemUtils/Source/OnlineSubsystemUtils/Private/OnlineBeaconClient.cpp
@@ -445,7 +445,7 @@ void AOnlineBeaconClient::NotifyControlMessage(UNetConnection* Connection, uint8
 				// send NMT_Login
 				Connection->ClientResponse = TEXT("0");
 				FNetControlMessage<NMT_Login>::Send(Connection, Connection->ClientResponse, URLString, Connection->PlayerId, OnlinePlatformNameString);
-				Connection->FlushNet();
+				NetDriver->ServerConnection->FlushNet();
 			}
 			else
 			{
@@ -464,7 +464,7 @@ void AOnlineBeaconClient::NotifyControlMessage(UNetConnection* Connection, uint8
 				if (!BeaconType.IsEmpty())
 				{
 					FNetControlMessage<NMT_BeaconJoin>::Send(Connection, BeaconType, Connection->PlayerId);
-					Connection->FlushNet();
+					NetDriver->ServerConnection->FlushNet();
 				}
 				else
 				{
@@ -486,7 +486,6 @@ void AOnlineBeaconClient::NotifyControlMessage(UNetConnection* Connection, uint8
 
 						FString BeaconType = GetBeaconType();
 						FNetControlMessage<NMT_BeaconNetGUIDAck>::Send(Connection, BeaconType);
-						Connection->FlushNet();
 						// Server will send ClientOnConnected() when it gets this control message
 
 						// Fail safe for connection to server but no client connection RPC
diff --git a/Engine/Plugins/Online/OnlineSubsystemUtils/Source/OnlineSubsystemUtils/Private/OnlineBeaconHost.cpp b/Engine/Plugins/Online/OnlineSubsystemUtils/Source/OnlineSubsystemUtils/Private/OnlineBeaconHost.cpp
--- a/Engine/Plugins/Online/OnlineSubsystemUtils/Source/OnlineSubsystemUtils/Private/OnlineBeaconHost.cpp
+++ b/Engine/Plugins/Online/OnlineSubsystemUtils/Source/OnlineSubsystemUtils/Private/OnlineBeaconHost.cpp
@@ -472,7 +472,6 @@ bool AOnlineBeaconHost::HandleControlMessage(UNetConnection* Connection, uint8 M
 			ClientActors.Add(NewClientActor);
 
 			FNetControlMessage<NMT_BeaconAssignGUID>::Send(Connection, NetGUID);
-			Connection->FlushNet(true);
 		}
 		break;
 	case NMT_BeaconNetGUIDAck:
@@ -564,8 +563,6 @@ void AOnlineBeaconHost::FinishHandshake(UNetConnection* Connection, FString Beac
 
 	// Send an RPC to the client to open the actor channel and guarantee RPCs will work
 	ClientActor->ClientOnConnected();
-	Connection->FlushNet(true);
-
 	UE_LOG(LogBeacon, Log, TEXT("%s: Handshake complete."), *GetDebugName(Connection));
 	OnBeaconConnectedDelegate->ExecuteIfBound(ClientActor, Connection);
 }
diff --git a/Engine/Plugins/Online/OnlineSubsystemUtils/Source/OnlineSubsystemUtilsTest/OnlineSubsystemUtilsTest.Build.cs b/Engine/Plugins/Online/OnlineSubsystemUtils/Source/OnlineSubsystemUtilsTest/OnlineSubsystemUtilsTest.Build.cs
deleted file mode 100644
diff --git a/Engine/Plugins/Online/OnlineSubsystemUtils/Source/OnlineSubsystemUtilsTest/Private/OnlineBeaconUnitTest.cpp b/Engine/Plugins/Online/OnlineSubsystemUtils/Source/OnlineSubsystemUtilsTest/Private/OnlineBeaconUnitTest.cpp
deleted file mode 100644
diff --git a/Engine/Plugins/Online/OnlineSubsystemUtils/Source/OnlineSubsystemUtilsTest/Private/OnlineBeaconUnitTestClient.cpp b/Engine/Plugins/Online/OnlineSubsystemUtils/Source/OnlineSubsystemUtilsTest/Private/OnlineBeaconUnitTestClient.cpp
deleted file mode 100644
diff --git a/Engine/Plugins/Online/OnlineSubsystemUtils/Source/OnlineSubsystemUtilsTest/Private/OnlineBeaconUnitTestClient.h b/Engine/Plugins/Online/OnlineSubsystemUtils/Source/OnlineSubsystemUtilsTest/Private/OnlineBeaconUnitTestClient.h
deleted file mode 100644
diff --git a/Engine/Plugins/Online/OnlineSubsystemUtils/Source/OnlineSubsystemUtilsTest/Private/OnlineBeaconUnitTestHost.cpp b/Engine/Plugins/Online/OnlineSubsystemUtils/Source/OnlineSubsystemUtilsTest/Private/OnlineBeaconUnitTestHost.cpp
deleted file mode 100644
diff --git a/Engine/Plugins/Online/OnlineSubsystemUtils/Source/OnlineSubsystemUtilsTest/Private/OnlineBeaconUnitTestHost.h b/Engine/Plugins/Online/OnlineSubsystemUtils/Source/OnlineSubsystemUtilsTest/Private/OnlineBeaconUnitTestHost.h
deleted file mode 100644
diff --git a/Engine/Plugins/Online/OnlineSubsystemUtils/Source/OnlineSubsystemUtilsTest/Private/OnlineBeaconUnitTestHostObject.cpp b/Engine/Plugins/Online/OnlineSubsystemUtils/Source/OnlineSubsystemUtilsTest/Private/OnlineBeaconUnitTestHostObject.cpp
deleted file mode 100644
diff --git a/Engine/Plugins/Online/OnlineSubsystemUtils/Source/OnlineSubsystemUtilsTest/Private/OnlineBeaconUnitTestHostObject.h b/Engine/Plugins/Online/OnlineSubsystemUtils/Source/OnlineSubsystemUtilsTest/Private/OnlineBeaconUnitTestHostObject.h
deleted file mode 100644
diff --git a/Engine/Plugins/Online/OnlineSubsystemUtils/Source/OnlineSubsystemUtilsTest/Private/OnlineBeaconUnitTestNetConnection.cpp b/Engine/Plugins/Online/OnlineSubsystemUtils/Source/OnlineSubsystemUtilsTest/Private/OnlineBeaconUnitTestNetConnection.cpp
deleted file mode 100644
diff --git a/Engine/Plugins/Online/OnlineSubsystemUtils/Source/OnlineSubsystemUtilsTest/Private/OnlineBeaconUnitTestNetConnection.h b/Engine/Plugins/Online/OnlineSubsystemUtils/Source/OnlineSubsystemUtilsTest/Private/OnlineBeaconUnitTestNetConnection.h
deleted file mode 100644
diff --git a/Engine/Plugins/Online/OnlineSubsystemUtils/Source/OnlineSubsystemUtilsTest/Private/OnlineBeaconUnitTestNetDriver.cpp b/Engine/Plugins/Online/OnlineSubsystemUtils/Source/OnlineSubsystemUtilsTest/Private/OnlineBeaconUnitTestNetDriver.cpp
deleted file mode 100644
diff --git a/Engine/Plugins/Online/OnlineSubsystemUtils/Source/OnlineSubsystemUtilsTest/Private/OnlineBeaconUnitTestNetDriver.h b/Engine/Plugins/Online/OnlineSubsystemUtils/Source/OnlineSubsystemUtilsTest/Private/OnlineBeaconUnitTestNetDriver.h
deleted file mode 100644
diff --git a/Engine/Plugins/Online/OnlineSubsystemUtils/Source/OnlineSubsystemUtilsTest/Private/OnlineBeaconUnitTestSocketSubsystem.cpp b/Engine/Plugins/Online/OnlineSubsystemUtils/Source/OnlineSubsystemUtilsTest/Private/OnlineBeaconUnitTestSocketSubsystem.cpp
deleted file mode 100644
diff --git a/Engine/Plugins/Online/OnlineSubsystemUtils/Source/OnlineSubsystemUtilsTest/Private/OnlineBeaconUnitTestSocketSubsystem.h b/Engine/Plugins/Online/OnlineSubsystemUtils/Source/OnlineSubsystemUtilsTest/Private/OnlineBeaconUnitTestSocketSubsystem.h
deleted file mode 100644
diff --git a/Engine/Plugins/Online/OnlineSubsystemUtils/Source/OnlineSubsystemUtilsTest/Private/OnlineBeaconUnitTestUtils.cpp b/Engine/Plugins/Online/OnlineSubsystemUtils/Source/OnlineSubsystemUtilsTest/Private/OnlineBeaconUnitTestUtils.cpp
deleted file mode 100644
diff --git a/Engine/Plugins/Online/OnlineSubsystemUtils/Source/OnlineSubsystemUtilsTest/Private/OnlineSubsystemUtilsTestModule.cpp b/Engine/Plugins/Online/OnlineSubsystemUtils/Source/OnlineSubsystemUtilsTest/Private/OnlineSubsystemUtilsTestModule.cpp
deleted file mode 100644
diff --git a/Engine/Plugins/Online/OnlineSubsystemUtils/Source/OnlineSubsystemUtilsTest/Public/OnlineBeaconUnitTestUtils.h b/Engine/Plugins/Online/OnlineSubsystemUtils/Source/OnlineSubsystemUtilsTest/Public/OnlineBeaconUnitTestUtils.h
deleted file mode 100644
diff --git a/Engine/Source/Runtime/Engine/Private/DataChannel.cpp b/Engine/Source/Runtime/Engine/Private/DataChannel.cpp
--- a/Engine/Source/Runtime/Engine/Private/DataChannel.cpp
+++ b/Engine/Source/Runtime/Engine/Private/DataChannel.cpp
@@ -1621,7 +1621,6 @@ void UChannel::AddedToChannelPool()
 -----------------------------------------------------------------------------*/
 
 const TCHAR* FNetControlMessageInfo::Names[256];
-bool FNetControlMessageInfo::bInitialized = false;
 
 // control channel message implementation
 IMPLEMENT_CONTROL_CHANNEL_MESSAGE(Hello);
diff --git a/Engine/Source/Runtime/Engine/Public/Net/DataChannel.h b/Engine/Source/Runtime/Engine/Public/Net/DataChannel.h
--- a/Engine/Source/Runtime/Engine/Public/Net/DataChannel.h
+++ b/Engine/Source/Runtime/Engine/Public/Net/DataChannel.h
@@ -51,6 +51,7 @@ public:
 private:
 	static void CheckInitialized()
 	{
+		static bool bInitialized = false;
 		if (!bInitialized)
 		{
 			for (int32 i = 0; i < UE_ARRAY_COUNT(Names); i++)
@@ -88,7 +89,6 @@ private:
 	
 	static constexpr int MaxNames = TNumericLimits<uint8>::Max() + 1;
 	static ENGINE_API const TCHAR* Names[MaxNames];
-	static ENGINE_API bool bInitialized;
 };
 
 #define DEFINE_CONTROL_CHANNEL_MESSAGE(Name, Index, ...) \
