Update the CreateSession BP node to include a "Use Lobbies If Available" boolean.

diff --git a/Engine/Plugins/Online/OnlineSubsystemUtils/Source/OnlineSubsystemUtils/Classes/CreateSessionCallbackProxy.h b/Engine/Plugins/Online/OnlineSubsystemUtils/Source/OnlineSubsystemUtils/Classes/CreateSessionCallbackProxy.h
--- a/Engine/Plugins/Online/OnlineSubsystemUtils/Source/OnlineSubsystemUtils/Classes/CreateSessionCallbackProxy.h
+++ b/Engine/Plugins/Online/OnlineSubsystemUtils/Source/OnlineSubsystemUtils/Classes/CreateSessionCallbackProxy.h
@@ -22,7 +22,7 @@ class UCreateSessionCallbackProxy : public UOnlineBlueprintCallProxyBase
 
 	// Creates a session with the default online subsystem
 	UFUNCTION(BlueprintCallable, meta=(BlueprintInternalUseOnly = "true", WorldContext="WorldContextObject"), Category = "Online|Session")
-	static UCreateSessionCallbackProxy* CreateSession(UObject* WorldContextObject, class APlayerController* PlayerController, int32 PublicConnections, bool bUseLAN);
+	static UCreateSessionCallbackProxy* CreateSession(UObject* WorldContextObject, class APlayerController* PlayerController, int32 PublicConnections, bool bUseLAN, bool bUseLobbiesIfAvailable = true);
 
 	// UOnlineBlueprintCallProxyBase interface
 	virtual void Activate() override;
@@ -54,6 +54,9 @@ private:
 	// Whether or not to search LAN
 	bool bUseLAN;
 
+	// Whether or not to use lobbies (as opposed to using a raw server list)
+	bool bUseLobbiesIfAvailable;
+
 	// The world context object in which this call is taking place
 	UObject* WorldContextObject;
 };
diff --git a/Engine/Plugins/Online/OnlineSubsystemUtils/Source/OnlineSubsystemUtils/Private/CreateSessionCallbackProxy.cpp b/Engine/Plugins/Online/OnlineSubsystemUtils/Source/OnlineSubsystemUtils/Private/CreateSessionCallbackProxy.cpp
--- a/Engine/Plugins/Online/OnlineSubsystemUtils/Source/OnlineSubsystemUtils/Private/CreateSessionCallbackProxy.cpp
+++ b/Engine/Plugins/Online/OnlineSubsystemUtils/Source/OnlineSubsystemUtils/Private/CreateSessionCallbackProxy.cpp
@@ -20,12 +20,13 @@ UCreateSessionCallbackProxy::UCreateSessionCallbackProxy(const FObjectInitialize
 {
 }
 
-UCreateSessionCallbackProxy* UCreateSessionCallbackProxy::CreateSession(UObject* WorldContextObject, class APlayerController* PlayerController, int32 PublicConnections, bool bUseLAN)
+UCreateSessionCallbackProxy* UCreateSessionCallbackProxy::CreateSession(UObject* WorldContextObject, class APlayerController* PlayerController, int32 PublicConnections, bool bUseLAN, bool bUseLobbiesIfAvailable)
 {
 	UCreateSessionCallbackProxy* Proxy = NewObject<UCreateSessionCallbackProxy>();
 	Proxy->PlayerControllerWeakPtr = PlayerController;
 	Proxy->NumPublicConnections = PublicConnections;
 	Proxy->bUseLAN = bUseLAN;
+	Proxy->bUseLobbiesIfAvailable = bUseLobbiesIfAvailable;
 	Proxy->WorldContextObject = WorldContextObject;
 	return Proxy;
 }
@@ -49,6 +50,7 @@ void UCreateSessionCallbackProxy::Activate()
 			Settings.bIsLANMatch = bUseLAN;
 			Settings.bUsesPresence = true;
 			Settings.bAllowJoinViaPresence = true;
+			Settings.bUseLobbiesIfAvailable = bUseLobbiesIfAvailable;
 
 			Sessions->CreateSession(*Helper.UserID, NAME_GameSession, Settings);
 
