Fix PIE online credentials settings paste bug

When PIE online credentials are changed in editor settings under Level Editor > Play Credentials, checks are run afterward, this includes copy and paste changes. However, none of these checks ran if you copied a whole login entry (User Id, Password, and Type) by right-clicking and copying the parent row of one and doing the same but pasting on another parent row. This change fixes that and the checks are run even if the paste targets the parent.

diff --git a/Engine/Plugins/Online/OnlineSubsystemUtils/Source/OnlineSubsystemUtils/Private/OnlinePIESettings.cpp b/Engine/Plugins/Online/OnlineSubsystemUtils/Source/OnlineSubsystemUtils/Private/OnlinePIESettings.cpp
--- a/Engine/Plugins/Online/OnlineSubsystemUtils/Source/OnlineSubsystemUtils/Private/OnlinePIESettings.cpp
+++ b/Engine/Plugins/Online/OnlineSubsystemUtils/Source/OnlineSubsystemUtils/Private/OnlinePIESettings.cpp
@@ -118,7 +118,11 @@ void UOnlinePIESettings::PostEditChangeProperty(FPropertyChangedEvent& PropertyC
 		else if (MemberPropName == GET_MEMBER_NAME_CHECKED(UOnlinePIESettings, Logins))
 		{
 			FName SubPropName = PropertyChangedEvent.Property->GetFName();
-			if (SubPropName == GET_MEMBER_NAME_CHECKED(FPIELoginSettingsInternal, Id))
+
+			// If we paste on top of the whole login entry, all fields will have changed and need their checks run.
+			const bool bPastedAllValues = SubPropName == GET_MEMBER_NAME_CHECKED(UOnlinePIESettings, Logins);
+			
+			if (bPastedAllValues || SubPropName == GET_MEMBER_NAME_CHECKED(FPIELoginSettingsInternal, Id))
 			{
 				TSet<FString> Ids;
 				for (FPIELoginSettingsInternal& Login : Logins)
@@ -140,7 +144,8 @@ void UOnlinePIESettings::PostEditChangeProperty(FPropertyChangedEvent& PropertyC
 					}
 				}
 			}
-			else if (SubPropName == GET_MEMBER_NAME_CHECKED(FPIELoginSettingsInternal, Token))
+
+			if (bPastedAllValues || SubPropName == GET_MEMBER_NAME_CHECKED(FPIELoginSettingsInternal, Token))
 			{
 				for (FPIELoginSettingsInternal& Login : Logins)
 				{
@@ -150,7 +155,8 @@ void UOnlinePIESettings::PostEditChangeProperty(FPropertyChangedEvent& PropertyC
 					Login.Encrypt();
 				}
 			}
-			else if (SubPropName == GET_MEMBER_NAME_CHECKED(FPIELoginSettingsInternal, Type))
+
+			if (bPastedAllValues || SubPropName == GET_MEMBER_NAME_CHECKED(FPIELoginSettingsInternal, Type))
 			{
 				for (FPIELoginSettingsInternal& Login : Logins)
 				{
