Refactor to remove usage of presence parameters in session search related logic, to follow SEARCH_PRESENCE deprecation

diff --git a/Engine/Plugins/Online/OnlineSubsystemUtils/Source/OnlineSubsystemUtils/Private/OnlineSubsystemUtils.cpp b/Engine/Plugins/Online/OnlineSubsystemUtils/Source/OnlineSubsystemUtils/Private/OnlineSubsystemUtils.cpp
--- a/Engine/Plugins/Online/OnlineSubsystemUtils/Source/OnlineSubsystemUtils/Private/OnlineSubsystemUtils.cpp
+++ b/Engine/Plugins/Online/OnlineSubsystemUtils/Source/OnlineSubsystemUtils/Private/OnlineSubsystemUtils.cpp
@@ -683,7 +683,6 @@ static bool OnlineExec( UWorld* InWorld, const TCHAR* Cmd, FOutputDevice& Ar )
 					else if (FParse::Command(&Cmd, TEXT("SESSIONHOST")))
 					{
 						bool bTestLAN = FParse::Command(&Cmd, TEXT("LAN")) ? true : false;
-						bool bTestPresence = FParse::Command(&Cmd, TEXT("PRESENCE")) ? true : false;
 
 						FOnlineSessionSettings SettingsOverride;
 
@@ -704,19 +703,18 @@ static bool OnlineExec( UWorld* InWorld, const TCHAR* Cmd, FOutputDevice& Ar )
 						}
 
 						// This class deletes itself once done
-						(new FTestSessionInterface(SubName, true))->Test(InWorld, bTestLAN, bTestPresence, false, SettingsOverride);
+						(new FTestSessionInterface(SubName, true))->Test(InWorld, bTestLAN, false, SettingsOverride);
 						bWasHandled = true;
 					}
 					// Spawn the object that will exercise all of the session methods as client
 					else if (FParse::Command(&Cmd, TEXT("SESSIONCLIENT")))
 					{
 						bool bTestLAN = FParse::Command(&Cmd, TEXT("LAN")) ? true : false;
-						bool bTestPresence = FParse::Command(&Cmd, TEXT("PRESENCE")) ? true : false;
 
 						FOnlineSessionSettings SettingsOverride;
 
 						// This class deletes itself once done
-						(new FTestSessionInterface(SubName, false))->Test(InWorld, bTestLAN, bTestPresence, false, SettingsOverride);
+						(new FTestSessionInterface(SubName, false))->Test(InWorld, bTestLAN, false, SettingsOverride);
 						bWasHandled = true;
 					}
 					else if (FParse::Command(&Cmd, TEXT("STARTMATCHMAKING")))
@@ -740,7 +738,7 @@ static bool OnlineExec( UWorld* InWorld, const TCHAR* Cmd, FOutputDevice& Ar )
 						}
 
 						// This class deletes itself once done
-						(new FTestSessionInterface(SubName, false))->Test(InWorld, false, false, true, SettingsOverride);
+						(new FTestSessionInterface(SubName, false))->Test(InWorld, false, true, SettingsOverride);
 						bWasHandled = true;
 					}
 					else if (FParse::Command(&Cmd, TEXT("CLOUD")))
diff --git a/Engine/Plugins/Online/OnlineSubsystemUtils/Source/OnlineSubsystemUtils/Private/Tests/TestSessionInterface.cpp b/Engine/Plugins/Online/OnlineSubsystemUtils/Source/OnlineSubsystemUtils/Private/Tests/TestSessionInterface.cpp
--- a/Engine/Plugins/Online/OnlineSubsystemUtils/Source/OnlineSubsystemUtils/Private/Tests/TestSessionInterface.cpp
+++ b/Engine/Plugins/Online/OnlineSubsystemUtils/Source/OnlineSubsystemUtils/Private/Tests/TestSessionInterface.cpp
@@ -16,7 +16,7 @@
 class TestOnlineGameSettings : public FOnlineSessionSettings
 {
  public:
-	TestOnlineGameSettings(bool bTestingLAN = false, bool bTestingPresence = false, const FOnlineSessionSettings& SettingsOverride = FOnlineSessionSettings())
+	TestOnlineGameSettings(bool bTestingLAN = false, const FOnlineSessionSettings& SettingsOverride = FOnlineSessionSettings())
 	{
 		NumPublicConnections = 10;
 		NumPrivateConnections = 0;
@@ -24,7 +24,6 @@ class TestOnlineGameSettings : public FOnlineSessionSettings
 		bShouldAdvertise = true;
 		bAllowJoinInProgress = true;
 		bAllowInvites = true;
-		bUsesPresence = bTestingPresence;
 		bAllowJoinViaPresence = true;
 		bAllowJoinViaPresenceFriendsOnly = false;
 
@@ -84,7 +83,7 @@ class TestOnlineGameSettings : public FOnlineSessionSettings
 class TestOnlineSearchSettings : public FOnlineSessionSearch
 {
 public:
-	TestOnlineSearchSettings(bool bSearchingLAN = false, bool bSearchingPresence = false, const FOnlineSessionSettings& SettingsOverride = FOnlineSessionSettings())
+	TestOnlineSearchSettings(bool bSearchingLAN = false, const FOnlineSessionSettings& SettingsOverride = FOnlineSessionSettings())
 	{
 		bIsLanQuery = bSearchingLAN;
 		MaxSearchResults = 10;
@@ -94,13 +93,6 @@ public:
 		QuerySettings.Set(FName(TEXT("TESTSETTING2")), (float)5.0f, EOnlineComparisonOp::Equals, 1);
 		QuerySettings.Set(FName(TEXT("TESTSETTING3")), FString(TEXT("Hello")), EOnlineComparisonOp::Equals, 2);
 
-		if (bSearchingPresence)
-		{
-PRAGMA_DISABLE_DEPRECATION_WARNINGS
-			QuerySettings.Set(SEARCH_PRESENCE, true, EOnlineComparisonOp::Equals);
-PRAGMA_ENABLE_DEPRECATION_WARNINGS
-		}
-
 		for (FSessionSettings::TConstIterator It(SettingsOverride.Settings); It; ++It)
 		{
 			FName Key = It.Key();
@@ -121,7 +113,7 @@ PRAGMA_ENABLE_DEPRECATION_WARNINGS
 	}
 };
 
-void FTestSessionInterface::Test(UWorld* InWorld, bool bTestLAN, bool bIsPresence, bool bIsMatchmaking, const FOnlineSessionSettings& SettingsOverride)
+void FTestSessionInterface::Test(UWorld* InWorld, bool bTestLAN, bool bIsMatchmaking, const FOnlineSessionSettings& SettingsOverride)
 {
 	IOnlineSubsystem* OnlineSub = Online::GetSubsystem(InWorld, FName(*Subsystem));
 	check(OnlineSub);
@@ -171,10 +163,10 @@ void FTestSessionInterface::Test(UWorld* InWorld, bool bTestLAN, bool bIsPresenc
 
 	if (bIsMatchmaking)
 	{
-		HostSettings = MakeShareable(new TestOnlineGameSettings(bTestLAN, bIsPresence, SettingsOverride));
+		HostSettings = MakeShareable(new TestOnlineGameSettings(bTestLAN, SettingsOverride));
 		HostSettings->AddWorldSettings(InWorld);
 
-		SearchSettings = MakeShareable(new TestOnlineSearchSettings(bTestLAN, bIsPresence, SettingsOverride));
+		SearchSettings = MakeShareable(new TestOnlineSearchSettings(bTestLAN, SettingsOverride));
 		TSharedRef<FOnlineSessionSearch> SearchSettingsRef = SearchSettings.ToSharedRef();
 
 		OnMatchmakingCompleteDelegateHandle = SessionInt->AddOnMatchmakingCompleteDelegate_Handle(OnMatchmakingCompleteDelegate);
@@ -187,7 +179,7 @@ void FTestSessionInterface::Test(UWorld* InWorld, bool bTestLAN, bool bIsPresenc
 	// Setup sessions
 	else if (bIsHost)
 	{
-		HostSettings = MakeShareable(new TestOnlineGameSettings(bTestLAN, bIsPresence, SettingsOverride));
+		HostSettings = MakeShareable(new TestOnlineGameSettings(bTestLAN, SettingsOverride));
 		HostSettings->AddWorldSettings(InWorld);
 
 		OnCreateSessionCompleteDelegateHandle = SessionInt->AddOnCreateSessionCompleteDelegate_Handle(OnCreateSessionCompleteDelegate);
@@ -195,7 +187,7 @@ void FTestSessionInterface::Test(UWorld* InWorld, bool bTestLAN, bool bIsPresenc
 	}
 	else
 	{
-		SearchSettings = MakeShareable(new TestOnlineSearchSettings(bTestLAN, bIsPresence, SettingsOverride));
+		SearchSettings = MakeShareable(new TestOnlineSearchSettings(bTestLAN, SettingsOverride));
 		TSharedRef<FOnlineSessionSearch> SearchSettingsRef = SearchSettings.ToSharedRef();
 
 		OnFindSessionsCompleteDelegateHandle = SessionInt->AddOnFindSessionsCompleteDelegate_Handle(OnFindSessionsCompleteDelegate);
@@ -411,7 +403,7 @@ void FTestSessionInterface::OnFindSessionByIdComplete(int32 LocalUserNum, bool b
 
 	DumpSession(&SearchResult.Session);
 
-	SearchSettings = MakeShared<TestOnlineSearchSettings>(false, false);
+	SearchSettings = MakeShared<TestOnlineSearchSettings>(false);
 	SearchSettings->SearchResults.Emplace(SearchResult);
 }
 
diff --git a/Engine/Plugins/Online/OnlineSubsystemUtils/Source/OnlineSubsystemUtils/Private/Tests/TestSessionInterface.h b/Engine/Plugins/Online/OnlineSubsystemUtils/Source/OnlineSubsystemUtils/Private/Tests/TestSessionInterface.h
--- a/Engine/Plugins/Online/OnlineSubsystemUtils/Source/OnlineSubsystemUtils/Private/Tests/TestSessionInterface.h
+++ b/Engine/Plugins/Online/OnlineSubsystemUtils/Source/OnlineSubsystemUtils/Private/Tests/TestSessionInterface.h
@@ -351,7 +351,7 @@ class TestOnlineSearchSettings;
 	 *z
 	 * @param bTestLAN setup settings for LAN test
 	 */
-	void Test(UWorld* InWorld, bool bTestLAN, bool bIsPresence, bool bIsMatchmaking, const FOnlineSessionSettings& SettingsOverride);
+	void Test(UWorld* InWorld, bool bTestLAN, bool bIsMatchmaking, const FOnlineSessionSettings& SettingsOverride);
 
 	/**
 	 * Clear out any existing delegates created
