Fix crash when PauseBeaconRequests is called while NetDriver is not set.

#rnx

diff --git a/Engine/Plugins/Online/OnlineSubsystemUtils/Source/OnlineSubsystemUtils/Public/OnlineBeacon.h b/Engine/Plugins/Online/OnlineSubsystemUtils/Source/OnlineSubsystemUtils/Public/OnlineBeacon.h
--- a/Engine/Plugins/Online/OnlineSubsystemUtils/Source/OnlineSubsystemUtils/Public/OnlineBeacon.h
+++ b/Engine/Plugins/Online/OnlineSubsystemUtils/Source/OnlineSubsystemUtils/Public/OnlineBeacon.h
@@ -75,6 +75,11 @@ class AOnlineBeacon : public AActor, public FNetworkNotify
 	 */
 	void PauseBeaconRequests(bool bPause)
 	{
+		if (!NetDriver)
+		{
+			return;
+		}
+
 		if (bPause)
 		{
 			UE_LOG(LogBeacon, Verbose, TEXT("All Beacon Requests Paused."));
