Distinguish between the spawn failing because no spawner is registered and the spawner returning null

#rnx

diff --git a/Engine/Plugins/Online/OnlineSubsystemUtils/Source/OnlineSubsystemUtils/Private/OnlineBeaconHost.cpp b/Engine/Plugins/Online/OnlineSubsystemUtils/Source/OnlineSubsystemUtils/Private/OnlineBeaconHost.cpp
--- a/Engine/Plugins/Online/OnlineSubsystemUtils/Source/OnlineSubsystemUtils/Private/OnlineBeaconHost.cpp
+++ b/Engine/Plugins/Online/OnlineSubsystemUtils/Source/OnlineSubsystemUtils/Private/OnlineBeaconHost.cpp
@@ -445,11 +445,17 @@ bool AOnlineBeaconHost::HandleControlMessage(UNetConnection* Connection, uint8 M
 
 			// spawn the beacon actor for this client
 			AOnlineBeaconClient* NewClientActor = nullptr;
-			FOnBeaconSpawned* OnBeaconSpawnedDelegate = OnBeaconSpawnedMapping.Find(BeaconType);
-			if (OnBeaconSpawnedDelegate && OnBeaconSpawnedDelegate->IsBound())
+			if (FOnBeaconSpawned* OnBeaconSpawnedDelegate = OnBeaconSpawnedMapping.Find(BeaconType);
+				OnBeaconSpawnedDelegate && OnBeaconSpawnedDelegate->IsBound())
 			{
 				NewClientActor = OnBeaconSpawnedDelegate->Execute(Connection);
 			}
+			else
+			{
+				static const FText ErrorTxt = NSLOCTEXT("NetworkErrors", "BeaconSpawnFailureError.UnknownSpawner", "Join failure, Couldn't spawn client beacon actor (unknown beacon spawner).");
+				SendFailurePacket(Connection, ENetCloseResult::BeaconSpawnFailureError, ErrorTxt);
+				return false;
+			}
 
 			// make sure it spawned correctly
 			if (NewClientActor == nullptr || BeaconType != NewClientActor->GetBeaconType())
