Make sure to destroy OnlineServices instances in PIE before destroying OnlineSubsystem instances.

#tests Tried in PIE with different type of game modes.
[RN] Make sure to destroy corresponding OnlineServices instances when stop PIE.

diff --git a/Engine/Plugins/Online/OnlineServices/Source/OnlineServicesInterface/Private/Online/OnlineServices.cpp b/Engine/Plugins/Online/OnlineServices/Source/OnlineServicesInterface/Private/Online/OnlineServices.cpp
--- a/Engine/Plugins/Online/OnlineServices/Source/OnlineServicesInterface/Private/Online/OnlineServices.cpp
+++ b/Engine/Plugins/Online/OnlineServices/Source/OnlineServicesInterface/Private/Online/OnlineServices.cpp
@@ -85,4 +85,9 @@ void DestroyAllNamedServices(EOnlineServices OnlineServices)
 	FOnlineServicesRegistry::Get().DestroyAllNamedServicesInstances(OnlineServices);
 }
 
+void DestroyAllServicesWithName(FName InstanceName)
+{
+	FOnlineServicesRegistry::Get().DestroyAllServicesInstancesWithName(InstanceName);
+}
+
 /* UE::Online */ }
diff --git a/Engine/Plugins/Online/OnlineServices/Source/OnlineServicesInterface/Private/Online/OnlineServicesRegistry.cpp b/Engine/Plugins/Online/OnlineServices/Source/OnlineServicesInterface/Private/Online/OnlineServicesRegistry.cpp
--- a/Engine/Plugins/Online/OnlineServices/Source/OnlineServicesInterface/Private/Online/OnlineServicesRegistry.cpp
+++ b/Engine/Plugins/Online/OnlineServices/Source/OnlineServicesInterface/Private/Online/OnlineServicesRegistry.cpp
@@ -152,6 +152,23 @@ void FOnlineServicesRegistry::DestroyAllNamedServicesInstances(EOnlineServices O
 	}
 }
 
+void FOnlineServicesRegistry::DestroyAllServicesInstancesWithName(FName InstanceName)
+{
+	for (auto It = NamedServiceInstances.CreateIterator(); It; ++It)
+	{
+		if (TSharedRef<IOnlineServices>* Service = It->Value.Find(InstanceName))
+		{
+			(*Service)->Destroy();
+			It->Value.Remove(InstanceName);
+		}
+
+		if (It->Value.IsEmpty())
+		{
+			It.RemoveCurrent();
+		}
+	}
+}
+
 TSharedPtr<IOnlineServices> FOnlineServicesRegistry::CreateServices(EOnlineServices OnlineServices, FName InstanceName)
 {
 	OnlineServices = ResolveServiceName(OnlineServices);
diff --git a/Engine/Plugins/Online/OnlineServices/Source/OnlineServicesInterface/Public/Online/OnlineServices.h b/Engine/Plugins/Online/OnlineServices/Source/OnlineServicesInterface/Public/Online/OnlineServices.h
--- a/Engine/Plugins/Online/OnlineServices/Source/OnlineServicesInterface/Public/Online/OnlineServices.h
+++ b/Engine/Plugins/Online/OnlineServices/Source/OnlineServicesInterface/Public/Online/OnlineServices.h
@@ -231,12 +231,19 @@ TSharedPtr<ServicesClass> GetServices(FName InstanceName = NAME_None)
 ONLINESERVICESINTERFACE_API void DestroyService(EOnlineServices OnlineServices = EOnlineServices::Default, FName InstanceName = NAME_None);
 
 /**
- * Destroy all instances of the online service specified
+ * Destroy all instances of the online services specified by EOnlineServices type
  *
  * @param OnlineServices Type of online services to destroy all instances from
  */
 ONLINESERVICESINTERFACE_API void DestroyAllNamedServices(EOnlineServices OnlineServices);
 
+/**
+ * Destroy all instances of the online services with specified instance name
+ *
+ * @param InstanceName Name of online services to destroy all instances from
+ */
+ONLINESERVICESINTERFACE_API void DestroyAllServicesWithName(FName InstanceName);
+
 namespace Meta {
 // TODO: Move to OnlineServices_Meta.inl file?
 
diff --git a/Engine/Plugins/Online/OnlineServices/Source/OnlineServicesInterface/Public/Online/OnlineServicesRegistry.h b/Engine/Plugins/Online/OnlineServices/Source/OnlineServicesInterface/Public/Online/OnlineServicesRegistry.h
--- a/Engine/Plugins/Online/OnlineServices/Source/OnlineServicesInterface/Public/Online/OnlineServicesRegistry.h
+++ b/Engine/Plugins/Online/OnlineServices/Source/OnlineServicesInterface/Public/Online/OnlineServicesRegistry.h
@@ -96,6 +96,13 @@ public:
 	 */
 	ONLINESERVICESINTERFACE_API void DestroyAllNamedServicesInstances(EOnlineServices OnlineServices);
 
+	/**
+	 * Destroy all instances of a specific InstanceName
+	 *
+	 * @param InstanceName  Name of online services for the IOnlineServices instance
+	 */
+	ONLINESERVICESINTERFACE_API void DestroyAllServicesInstancesWithName(FName InstanceName);
+
 	/**
 	 * Create and initialize a new IOnlineServices instance
 	 *
diff --git a/Engine/Plugins/Online/OnlineSubsystemUtils/Source/OnlineSubsystemUtils/Private/OnlineEngineInterfaceImpl.cpp b/Engine/Plugins/Online/OnlineSubsystemUtils/Source/OnlineSubsystemUtils/Private/OnlineEngineInterfaceImpl.cpp
--- a/Engine/Plugins/Online/OnlineSubsystemUtils/Source/OnlineSubsystemUtils/Private/OnlineEngineInterfaceImpl.cpp
+++ b/Engine/Plugins/Online/OnlineSubsystemUtils/Source/OnlineSubsystemUtils/Private/OnlineEngineInterfaceImpl.cpp
@@ -11,6 +11,30 @@
 
 #include UE_INLINE_GENERATED_CPP_BY_NAME(OnlineEngineInterfaceImpl)
 
+namespace UOnlineEngineInterfaceImplPrivate
+{
+
+static FName ConvertToCompatibilityOnlineIdentifier(FName OnlineIdentifier)
+{
+	FName CompatibilityOnlineIdentifier;
+
+	FString OnlineIdentifierStr = OnlineIdentifier.ToString();
+	int32 Index = -1;
+	bool bFound = OnlineIdentifierStr.FindChar(':', Index);
+	if (bFound)
+	{
+		CompatibilityOnlineIdentifier = FName(OnlineIdentifierStr.Mid(Index+1));
+	}
+	else
+	{
+		CompatibilityOnlineIdentifier = OnlineIdentifier;
+	}
+
+	return CompatibilityOnlineIdentifier;
+}
+
+}
+
 UOnlineEngineInterfaceImpl::UOnlineEngineInterfaceImpl(const FObjectInitializer& ObjectInitializer)
 	: Super(ObjectInitializer)
 	, VoiceSubsystemNameOverride(NAME_None)
@@ -56,6 +80,12 @@ bool UOnlineEngineInterfaceImpl::DoesInstanceExist(FName OnlineIdentifier)
 
 void UOnlineEngineInterfaceImpl::ShutdownOnlineSubsystem(FName OnlineIdentifier)
 {
+	if (UOnlineEngineInterface* CompatibilityInterface = OnlineServicesCompatibilityInterface.Get())
+	{
+		FName CompatibilityOnlineIdentifier = UOnlineEngineInterfaceImplPrivate::ConvertToCompatibilityOnlineIdentifier(OnlineIdentifier);
+		CompatibilityInterface->ShutdownOnlineSubsystem(CompatibilityOnlineIdentifier);
+	}
+
 	IOnlineSubsystem* OnlineSub = IOnlineSubsystem::Get(OnlineIdentifier);
 	if (OnlineSub)
 	{
@@ -65,6 +95,12 @@ void UOnlineEngineInterfaceImpl::ShutdownOnlineSubsystem(FName OnlineIdentifier)
 
 void UOnlineEngineInterfaceImpl::DestroyOnlineSubsystem(FName OnlineIdentifier)
 {
+	if (UOnlineEngineInterface* CompatibilityInterface = OnlineServicesCompatibilityInterface.Get())
+	{
+		FName CompatibilityOnlineIdentifier = UOnlineEngineInterfaceImplPrivate::ConvertToCompatibilityOnlineIdentifier(OnlineIdentifier);
+		CompatibilityInterface->DestroyOnlineSubsystem(CompatibilityOnlineIdentifier);
+	}
+
 	IOnlineSubsystem::Destroy(OnlineIdentifier);
 }
 
diff --git a/Engine/Plugins/Online/OnlineSubsystemUtils/Source/OnlineSubsystemUtils/Private/OnlineEngineInterfaceImpl.h b/Engine/Plugins/Online/OnlineSubsystemUtils/Source/OnlineSubsystemUtils/Private/OnlineEngineInterfaceImpl.h
--- a/Engine/Plugins/Online/OnlineSubsystemUtils/Source/OnlineSubsystemUtils/Private/OnlineEngineInterfaceImpl.h
+++ b/Engine/Plugins/Online/OnlineSubsystemUtils/Source/OnlineSubsystemUtils/Private/OnlineEngineInterfaceImpl.h
@@ -28,6 +28,8 @@ public:
 	virtual bool IsLoaded(FName OnlineIdentifier) override;
 	virtual FName GetOnlineIdentifier(FWorldContext& WorldContext) override;
 	virtual bool DoesInstanceExist(FName OnlineIdentifier) override;
+
+	// NOTE: In OSSv1 it only shuts down the default type of subsystem instance corresponding to the identifier
 	virtual void ShutdownOnlineSubsystem(FName OnlineIdentifier) override;
 	virtual void DestroyOnlineSubsystem(FName OnlineIdentifier) override;
 	virtual FName GetDefaultOnlineSubsystemName() const override;
diff --git a/Engine/Plugins/Online/OnlineSubsystemUtils/Source/OnlineSubsystemUtils/Private/OnlineServicesEngineInterfaceImpl.cpp b/Engine/Plugins/Online/OnlineSubsystemUtils/Source/OnlineSubsystemUtils/Private/OnlineServicesEngineInterfaceImpl.cpp
--- a/Engine/Plugins/Online/OnlineSubsystemUtils/Source/OnlineSubsystemUtils/Private/OnlineServicesEngineInterfaceImpl.cpp
+++ b/Engine/Plugins/Online/OnlineSubsystemUtils/Source/OnlineSubsystemUtils/Private/OnlineServicesEngineInterfaceImpl.cpp
@@ -49,18 +49,13 @@ bool UOnlineServicesEngineInterfaceImpl::DoesInstanceExist(FName OnlineIdentifie
 
 void UOnlineServicesEngineInterfaceImpl::ShutdownOnlineSubsystem(FName OnlineIdentifier)
 {
-	// TODO:  Does this need to support multiple online service types?  Other accessors seem to not differentiate and just use the default
-	if (UE::Online::IOnlineServicesPtr OnlineServices = UE::Online::GetServices(UE::Online::EOnlineServices::Default, OnlineIdentifier))
-	{
-		// TODO:  Is this correct?  Or a new shutdown method?
-		OnlineServices->Destroy();
-	}
+	// While using OSS Adaptors, need to destroy OnlineServices instances here instead of shutting down, to 
+	// make sure the shared ptr references to OSSv1 interfaces in oss adaptors are removed
+	UE::Online::DestroyAllServicesWithName(OnlineIdentifier);
 }
 
 void UOnlineServicesEngineInterfaceImpl::DestroyOnlineSubsystem(FName OnlineIdentifier)
 {
-	// TODO:  Does this need to support multiple online service types?  Other accessors seem to not differentiate and just use the default
-	UE::Online::DestroyService(UE::Online::EOnlineServices::Default, OnlineIdentifier);
 }
 
 FName UOnlineServicesEngineInterfaceImpl::GetDefaultOnlineSubsystemName() const
diff --git a/Engine/Plugins/Online/OnlineSubsystemUtils/Source/OnlineSubsystemUtils/Private/OnlineServicesEngineInterfaceImpl.h b/Engine/Plugins/Online/OnlineSubsystemUtils/Source/OnlineSubsystemUtils/Private/OnlineServicesEngineInterfaceImpl.h
--- a/Engine/Plugins/Online/OnlineSubsystemUtils/Source/OnlineSubsystemUtils/Private/OnlineServicesEngineInterfaceImpl.h
+++ b/Engine/Plugins/Online/OnlineSubsystemUtils/Source/OnlineSubsystemUtils/Private/OnlineServicesEngineInterfaceImpl.h
@@ -28,6 +28,7 @@ public:
 	virtual bool IsLoaded(FName OnlineIdentifier) override;
 	virtual FName GetOnlineIdentifier(FWorldContext& WorldContext) override;
 	virtual bool DoesInstanceExist(FName OnlineIdentifier) override;
+	// NOTE: In OnlineServices it destroys all types of OnlineServices instances which contains this identifier
 	virtual void ShutdownOnlineSubsystem(FName OnlineIdentifier) override;
 	virtual void DestroyOnlineSubsystem(FName OnlineIdentifier) override;
 	virtual FName GetDefaultOnlineSubsystemName() const override;
