Move OnlineServices per-world instance name determination to its own function GetServicesInstanceName
Fix CommonConfig determining the online services name differently from other callers
#changelist validated

diff --git a/Engine/Plugins/Online/OnlineSubsystemUtils/Source/OnlineSubsystemUtils/Public/Online/OnlineServicesEngineUtils.h b/Engine/Plugins/Online/OnlineSubsystemUtils/Source/OnlineSubsystemUtils/Public/Online/OnlineServicesEngineUtils.h
--- a/Engine/Plugins/Online/OnlineSubsystemUtils/Source/OnlineSubsystemUtils/Public/Online/OnlineServicesEngineUtils.h
+++ b/Engine/Plugins/Online/OnlineSubsystemUtils/Source/OnlineSubsystemUtils/Public/Online/OnlineServicesEngineUtils.h
@@ -71,7 +71,7 @@ public:
 /** @return the single instance of the online services utils interface */
 ONLINESUBSYSTEMUTILS_API IOnlineServicesEngineUtils* GetServicesEngineUtils();
 
-inline IOnlineServicesPtr GetServices(const UWorld* World, EOnlineServices OnlineServices = EOnlineServices::Default)
+inline FName GetServicesInstanceName(const UWorld* World)
 {
 	FName Identifier;
 #if UE_EDITOR // at present, multiple worlds are only possible in the editor
@@ -82,7 +82,12 @@ inline IOnlineServicesPtr GetServices(const UWorld* World, EOnlineServices Onlin
 		Identifier = Utils->GetOnlineIdentifier(World);
 	}
 #endif
-	return GetServices(OnlineServices, Identifier);
+	return Identifier;
+}
+
+inline IOnlineServicesPtr GetServices(const UWorld* World, EOnlineServices OnlineServices = EOnlineServices::Default)
+{
+	return GetServices(OnlineServices, GetServicesInstanceName(World));
 }
 
 /* UE::Online */ }
