Added validation function to PartyBeaconState when adding a new reservation to allow child classes to perform additional validation on a reservation request before it's added

#tests Worked with QA to create a list of matchmaking flows that we expect to work. Tested each of them and verified that players were able to connect to the server when we expected them to and weren't able to connect when we didn't expect them to.
#rnx

diff --git a/Engine/Plugins/Online/OnlineSubsystemUtils/Source/OnlineSubsystemUtils/Private/PartyBeaconHost.cpp b/Engine/Plugins/Online/OnlineSubsystemUtils/Source/OnlineSubsystemUtils/Private/PartyBeaconHost.cpp
--- a/Engine/Plugins/Online/OnlineSubsystemUtils/Source/OnlineSubsystemUtils/Private/PartyBeaconHost.cpp
+++ b/Engine/Plugins/Online/OnlineSubsystemUtils/Source/OnlineSubsystemUtils/Private/PartyBeaconHost.cpp
@@ -676,27 +676,34 @@ EPartyReservationResult::Type APartyBeaconHost::AddPartyReservation(const FParty
 						{
 							if (State->CrossPlayAllowed(ReservationRequest))
 							{
-								if (State->AddReservation(ReservationRequest))
+								if (State->ValidateNewPartyReservation(ReservationRequest))
 								{
-									// Keep track of newly added players
-									for (const FPlayerReservation& PartyMember : ReservationRequest.PartyMembers)
+									if (State->AddReservation(ReservationRequest))
 									{
-										NewPlayerAdded(PartyMember);
-									}
+										// Keep track of newly added players
+										for (const FPlayerReservation& PartyMember : ReservationRequest.PartyMembers)
+										{
+											NewPlayerAdded(PartyMember);
+										}
 
-									SendReservationUpdates();
+										SendReservationUpdates();
 
-									NotifyReservationEventNextFrame(ReservationChanged);
-									if (State->IsBeaconFull())
+										NotifyReservationEventNextFrame(ReservationChanged);
+										if (State->IsBeaconFull())
+										{
+											NotifyReservationEventNextFrame(ReservationsFull);
+										}
+										Result = EPartyReservationResult::ReservationAccepted;
+									}
+									else
 									{
-										NotifyReservationEventNextFrame(ReservationsFull);
+										// Something wrong with team assignment
+										Result = EPartyReservationResult::IncorrectPlayerCount;
 									}
-									Result = EPartyReservationResult::ReservationAccepted;
 								}
 								else
 								{
-									// Something wrong with team assignment
-									Result = EPartyReservationResult::IncorrectPlayerCount;
+									Result = EPartyReservationResult::ReservationDenied;
 								}
 							}
 							else
diff --git a/Engine/Plugins/Online/OnlineSubsystemUtils/Source/OnlineSubsystemUtils/Public/PartyBeaconState.h b/Engine/Plugins/Online/OnlineSubsystemUtils/Source/OnlineSubsystemUtils/Public/PartyBeaconState.h
--- a/Engine/Plugins/Online/OnlineSubsystemUtils/Source/OnlineSubsystemUtils/Public/PartyBeaconState.h
+++ b/Engine/Plugins/Online/OnlineSubsystemUtils/Source/OnlineSubsystemUtils/Public/PartyBeaconState.h
@@ -565,6 +565,16 @@ class ONLINESUBSYSTEMUTILS_API UPartyBeaconState : public UObject
 	 */
 	virtual void DumpReservations() const;
 
+	/**
+	 * Allows child classes to add extra validation when a new reservation request comes in, 
+	 * for example to validate the new players are allowed in based on privacy settings
+	 * 
+	 * @param ReservationRequest The new reservation that will be added to the beacon's list of reservations
+	 *
+	 * @return True if the players in the reservation are allowed in the session
+	 */
+	virtual bool ValidateNewPartyReservation(const FPartyReservation& ReservationRequest) { return true; }
+
 protected:
 
 	/** Session tied to the beacon */
