Add support for custom reservation timeout logic.
#rnx

diff --git a/Engine/Plugins/Online/OnlineSubsystemUtils/Source/OnlineSubsystemUtils/Private/PartyBeaconHost.cpp b/Engine/Plugins/Online/OnlineSubsystemUtils/Source/OnlineSubsystemUtils/Private/PartyBeaconHost.cpp
--- a/Engine/Plugins/Online/OnlineSubsystemUtils/Source/OnlineSubsystemUtils/Private/PartyBeaconHost.cpp
+++ b/Engine/Plugins/Online/OnlineSubsystemUtils/Source/OnlineSubsystemUtils/Private/PartyBeaconHost.cpp
@@ -98,6 +98,17 @@ void APartyBeaconHost::SetTeamAssignmentMethod(FName NewAssignmentMethod)
 	}
 }
 
+bool APartyBeaconHost::IsPlayerActive(const FUniqueNetIdRepl& PlayerId) const
+{
+	IOnlineSessionPtr SessionInt = Online::GetSessionInterface(GetWorld());
+	if (SessionInt.IsValid())
+	{
+		return SessionInt->IsPlayerInSession(State->GetSessionName(), *PlayerId);
+	}
+
+	return false;
+}
+
 void APartyBeaconHost::Tick(float DeltaTime)
 {
 	if (State)
@@ -160,7 +171,7 @@ void APartyBeaconHost::Tick(float DeltaTime)
 							const bool bIsSessionOwner = Session->OwningUserId.IsValid() && (*Session->OwningUserId == *PlayerEntry.UniqueId);
 
 							// Determine if the player member is registered in the game session
-							if (SessionInt->IsPlayerInSession(SessionName, *PlayerEntry.UniqueId) ||
+							if (IsPlayerActive(PlayerEntry.UniqueId) ||
 								// Never timeout the session owner
 								bIsSessionOwner)
 							{
diff --git a/Engine/Plugins/Online/OnlineSubsystemUtils/Source/OnlineSubsystemUtils/Public/PartyBeaconHost.h b/Engine/Plugins/Online/OnlineSubsystemUtils/Source/OnlineSubsystemUtils/Public/PartyBeaconHost.h
--- a/Engine/Plugins/Online/OnlineSubsystemUtils/Source/OnlineSubsystemUtils/Public/PartyBeaconHost.h
+++ b/Engine/Plugins/Online/OnlineSubsystemUtils/Source/OnlineSubsystemUtils/Public/PartyBeaconHost.h
@@ -401,6 +401,8 @@ class ONLINESUBSYSTEMUTILS_API APartyBeaconHost : public AOnlineBeaconHostObject
 	 */
 	bool IsValidationStrRequired() const { return bIsValidationStrRequired; }
 
+	virtual bool IsPlayerActive(const FUniqueNetIdRepl& PlayerId) const;
+
 protected:
 
 	/** State of the beacon */
