- Fix MCP IsEnabled returning false when called before UObjectInitialized() and there are no configured PIE login creds just ones provided on command line

diff --git a/Engine/Plugins/Online/OnlineSubsystemUtils/Source/OnlineSubsystemUtils/Private/OnlinePIESettings.cpp b/Engine/Plugins/Online/OnlineSubsystemUtils/Source/OnlineSubsystemUtils/Private/OnlinePIESettings.cpp
--- a/Engine/Plugins/Online/OnlineSubsystemUtils/Source/OnlineSubsystemUtils/Private/OnlinePIESettings.cpp
+++ b/Engine/Plugins/Online/OnlineSubsystemUtils/Source/OnlineSubsystemUtils/Private/OnlinePIESettings.cpp
@@ -1,8 +1,10 @@
 // Copyright Epic Games, Inc. All Rights Reserved.
 
 #include "OnlinePIESettings.h"
-#include "OnlineSubsystem.h"
+
 #include "Misc/AES.h"
+#include "Misc/Parse.h"
+#include "OnlineSubsystem.h"
 
 #include UE_INLINE_GENERATED_CPP_BY_NAME(OnlinePIESettings)
 
@@ -85,6 +87,37 @@ void FPIELoginSettingsInternal::Decrypt()
 	}
 }
 
+FPIELoginSettingsInternal FPIELoginSettingsInternal::FromConfigString(const FString& ConfigString)
+{
+	FPIELoginSettingsInternal Result;
+	FParse::Value(*ConfigString, TEXT("Id="), Result.Id);
+	FParse::Value(*ConfigString, TEXT("Type="), Result.Type);
+	FParse::Value(*ConfigString, TEXT("Token="), Result.Token);
+
+	FString TokenBytesStr;
+	FParse::Value(*ConfigString, TEXT("TokenBytes="), TokenBytesStr, false);
+	if (!TokenBytesStr.IsEmpty())
+	{
+		check(TokenBytesStr[0] == TCHAR('('));
+		TokenBytesStr.LeftInline(FindMatchingClosingParenthesis(TokenBytesStr, 0));
+		TokenBytesStr.RemoveFromStart(TEXT("("));
+
+		TCHAR* ContextStr = nullptr;
+		const TCHAR* Elem = FCString::Strtok(TokenBytesStr.GetCharArray().GetData(), TEXT(","), &ContextStr);
+		while (Elem)
+		{
+			Result.TokenBytes.Emplace(FCString::Atoi(Elem));
+			Elem = FCString::Strtok(nullptr, TEXT(","), &ContextStr);
+		}
+		if (!Result.TokenBytes.IsEmpty())
+		{
+			Result.Decrypt();
+		}
+	}
+
+	return Result;
+}
+
 UOnlinePIESettings::UOnlinePIESettings(const FObjectInitializer& ObjectInitializer)
 	: Super(ObjectInitializer)
 {
diff --git a/Engine/Plugins/Online/OnlineSubsystemUtils/Source/OnlineSubsystemUtils/Private/OnlinePIESettings.h b/Engine/Plugins/Online/OnlineSubsystemUtils/Source/OnlineSubsystemUtils/Private/OnlinePIESettings.h
--- a/Engine/Plugins/Online/OnlineSubsystemUtils/Source/OnlineSubsystemUtils/Private/OnlinePIESettings.h
+++ b/Engine/Plugins/Online/OnlineSubsystemUtils/Source/OnlineSubsystemUtils/Private/OnlinePIESettings.h
@@ -40,6 +40,8 @@ public:
 	 * Decrypt the TokenBytes field into the Token field
 	 */
 	void Decrypt();
+
+	static FPIELoginSettingsInternal FromConfigString(const FString& ConfigString);
 };
 
 /**
diff --git a/Engine/Plugins/Online/OnlineSubsystemUtils/Source/OnlineSubsystemUtils/Private/OnlineSubsystemUtilsModule.cpp b/Engine/Plugins/Online/OnlineSubsystemUtils/Source/OnlineSubsystemUtils/Private/OnlineSubsystemUtilsModule.cpp
--- a/Engine/Plugins/Online/OnlineSubsystemUtils/Source/OnlineSubsystemUtils/Private/OnlineSubsystemUtilsModule.cpp
+++ b/Engine/Plugins/Online/OnlineSubsystemUtils/Source/OnlineSubsystemUtils/Private/OnlineSubsystemUtilsModule.cpp
@@ -1,11 +1,14 @@
 // Copyright Epic Games, Inc. All Rights Reserved.
 
 #include "OnlineSubsystemUtilsModule.h"
+
+#include "Misc/CommandLine.h"
+#include "Misc/ConfigCacheIni.h"
 #include "Online.h"
-#include "OnlineSubsystemUtils.h"
-#include "OnlineServicesEngineUtilsImpl.h"
 #include "OnlineDelegates.h"
 #include "OnlinePIESettings.h"
+#include "OnlineServicesEngineUtilsImpl.h"
+#include "OnlineSubsystemUtils.h"
 
 #if WITH_EDITOR
 #include "Misc/CommandLine.h"
@@ -124,9 +127,9 @@ private:
 public:	
 	virtual bool SupportsOnlinePIE() const override
 	{
-		check(UObjectInitialized());
-		const UOnlinePIESettings* OnlinePIESettings = GetDefault<UOnlinePIESettings>();
-		if ((OnlinePIESettings->bOnlinePIEEnabled && GetNumPIELogins() > 0) || HasAuthCommandLine())
+		bool bOnlinePIEEnabled = false;
+		GConfig->GetBool(TEXT("/Script/OnlineSubsystemUtils.OnlinePIESettings"), TEXT("bOnlinePIEEnabled"), bOnlinePIEEnabled, GEditorPerProjectIni);
+		if ((bOnlinePIEEnabled && GetNumPIELogins() > 0) || HasAuthCommandLine())
 		{
 			// If we can't get the identity interface then things are either not configured right or disabled
 			IOnlineIdentityPtr IdentityInt = Online::GetIdentityInterface();
@@ -149,23 +152,21 @@ public:
 
 	virtual bool IsOnlinePIEEnabled() const override
 	{
-		check(UObjectInitialized());
-		if (!bShouldTryOnlinePIE)
-		{
-			return false;
-		}
-		const UOnlinePIESettings* OnlinePIESettings = GetDefault<UOnlinePIESettings>();
-		return OnlinePIESettings->bOnlinePIEEnabled || HasAuthCommandLine();
+		bool bOnlinePIEEnabled = false;
+		GConfig->GetBool(TEXT("/Script/OnlineSubsystemUtils.OnlinePIESettings"), TEXT("bOnlinePIEEnabled"), bOnlinePIEEnabled, GEditorPerProjectIni);
+		return bShouldTryOnlinePIE && (bOnlinePIEEnabled || HasAuthCommandLine());
 	}
 
 	virtual int32 GetNumPIELogins() const override
 	{
-		check(UObjectInitialized());
-
 		int32 NumValidLogins = 0;
-		const UOnlinePIESettings* OnlinePIESettings = GetDefault<UOnlinePIESettings>();
-		for (const FPIELoginSettingsInternal& Login : OnlinePIESettings->Logins)
+
+		TArray<FString> LoginStrings;
+		GConfig->GetArray(TEXT("/Script/OnlineSubsystemUtils.OnlinePIESettings"), TEXT("Logins"), LoginStrings, GEditorPerProjectIni);
+
+		for (const FString& LoginString : LoginStrings)
 		{
+			FPIELoginSettingsInternal Login = FPIELoginSettingsInternal::FromConfigString(LoginString);
 			if (Login.IsValid())
 			{
 				NumValidLogins++;
@@ -181,13 +182,15 @@ public:
 
 	virtual void GetPIELogins(TArray<FOnlineAccountCredentials>& Logins) override
 	{
-		check(UObjectInitialized());
-		const UOnlinePIESettings* OnlinePIESettings = GetDefault<UOnlinePIESettings>();
-		if (OnlinePIESettings->Logins.Num() > 0)
+		TArray<FString> LoginStrings;
+		GConfig->GetArray(TEXT("/Script/OnlineSubsystemUtils.OnlinePIESettings"), TEXT("Logins"), LoginStrings, GEditorPerProjectIni);
+
+		if (LoginStrings.Num() > 0)
 		{
-			Logins.Empty(OnlinePIESettings->Logins.Num());
-			for (const FPIELoginSettingsInternal& Login : OnlinePIESettings->Logins)
+			Logins.Empty(LoginStrings.Num());
+			for (const FString& LoginString : LoginStrings)
 			{
+				FPIELoginSettingsInternal Login = FPIELoginSettingsInternal::FromConfigString(LoginString);
 				if (Login.IsValid())
 				{
 					Logins.Emplace(Login.Type, Login.Id, Login.Token);
diff --git a/Engine/Plugins/Online/OnlineSubsystemUtils/Source/OnlineSubsystemUtils/Public/OnlineSubsystemUtils.h b/Engine/Plugins/Online/OnlineSubsystemUtils/Source/OnlineSubsystemUtils/Public/OnlineSubsystemUtils.h
--- a/Engine/Plugins/Online/OnlineSubsystemUtils/Source/OnlineSubsystemUtils/Public/OnlineSubsystemUtils.h
+++ b/Engine/Plugins/Online/OnlineSubsystemUtils/Source/OnlineSubsystemUtils/Public/OnlineSubsystemUtils.h
@@ -185,14 +185,8 @@ namespace Online
 	/** @return the single instance of the online subsystem utils interface */
 	static IOnlineSubsystemUtils* GetUtils()
 	{
-		static const FName OnlineSubsystemModuleName = TEXT("OnlineSubsystemUtils");
-		FOnlineSubsystemUtilsModule* OSSUtilsModule = FModuleManager::GetModulePtr<FOnlineSubsystemUtilsModule>(OnlineSubsystemModuleName);
-		if (OSSUtilsModule != nullptr)
-		{
-			return OSSUtilsModule->GetUtils();
-		}
-
-		return nullptr;
+		const FOnlineSubsystemUtilsModule* OSSUtilsModule = FModuleManager::LoadModulePtr<FOnlineSubsystemUtilsModule>(TEXT("OnlineSubsystemUtils"));
+		return OSSUtilsModule ? OSSUtilsModule->GetUtils() : nullptr;
 	}
 
 	/**
