[AutoRTFM] Defer network activity until the completion of a transaction.

Inside an AutoRTFM transaction, we must postpone network sends until commit time, to account for the possibility of aborting the transaction. An abort will prevent any network traffic from occurring.

diff --git a/Engine/Plugins/Online/OnlineSubsystemUtils/Source/OnlineSubsystemUtils/Classes/IpConnection.h b/Engine/Plugins/Online/OnlineSubsystemUtils/Source/OnlineSubsystemUtils/Classes/IpConnection.h
--- a/Engine/Plugins/Online/OnlineSubsystemUtils/Source/OnlineSubsystemUtils/Classes/IpConnection.h
+++ b/Engine/Plugins/Online/OnlineSubsystemUtils/Source/OnlineSubsystemUtils/Classes/IpConnection.h
@@ -142,6 +142,9 @@ private:
 #endif
 
 private:
+	/** Performs the work of LowLevelSend once we have confirmed that we want to send data to remote. */
+	void SendToRemote(uint8* Data, int32 CountBits, FOutPacketTraits& Traits);
+
 	/** Sets the local socket pointer, and safely cleans up any references to old sockets */
 	void SetSocket_Local(const TSharedPtr<FSocket>& InSocket);
 
diff --git a/Engine/Plugins/Online/OnlineSubsystemUtils/Source/OnlineSubsystemUtils/Private/IpConnection.cpp b/Engine/Plugins/Online/OnlineSubsystemUtils/Source/OnlineSubsystemUtils/Private/IpConnection.cpp
--- a/Engine/Plugins/Online/OnlineSubsystemUtils/Source/OnlineSubsystemUtils/Private/IpConnection.cpp
+++ b/Engine/Plugins/Online/OnlineSubsystemUtils/Source/OnlineSubsystemUtils/Private/IpConnection.cpp
@@ -11,6 +11,7 @@ Notes:
 #include "Async/TaskGraphInterfaces.h"
 #include "IpNetDriver.h"
 #include "Engine/Engine.h"
+#include "AutoRTFM/AutoRTFM.h"
 
 #include "Sockets.h"
 #include "Net/NetworkProfiler.h"
@@ -417,7 +418,7 @@ void UIpConnection::LowLevelSend(void* Data, int32 CountBits, FOutPacketTraits&
 {
 	using namespace UE::Net::Private;
 
-	const uint8* DataToSend = reinterpret_cast<uint8*>(Data);
+	uint8* DataToSend = static_cast<uint8*>(Data);
 
 	// If the remote addr hasn't been set, we need to wait for it.
 	if (!RemoteAddr.IsValid() || !RemoteAddr->IsValid())
@@ -425,6 +426,42 @@ void UIpConnection::LowLevelSend(void* Data, int32 CountBits, FOutPacketTraits&
 		return;
 	}
 
+#if !UE_BUILD_SHIPPING
+	if (GNetBlockSend)
+	{
+		return;
+	}
+#endif
+
+	if (Resolver->ShouldBlockSend())
+	{
+		UE_LOG(LogNet, Verbose, TEXT("Skipping send task as we are waiting on the next resolution step"));
+		return;
+	}
+
+	if (AutoRTFM::IsClosed())
+	{
+		// Inside an AutoRTFM transaction, we must postpone network sends until commit time, to account for the
+		// possibility of aborting the transaction. An abort will prevent any network traffic from occurring.
+		int32 CountBytes = FMath::DivideAndRoundUp(CountBits, 8);
+
+		AutoRTFM::OnCommit([this, 
+			                DataCopy = TArray<uint8>(DataToSend, CountBytes),
+							CountBits,
+			                TraitsCopy = Traits]() mutable
+		{
+			SendToRemote(DataCopy.GetData(), CountBits, TraitsCopy);
+		});
+	}
+	else
+	{
+		SendToRemote(DataToSend, CountBits, Traits);
+	}
+}
+
+void UIpConnection::SendToRemote(uint8* DataToSend, int32 CountBits, FOutPacketTraits& Traits) {
+	using namespace UE::Net::Private;
+
 #if !UE_BUILD_SHIPPING
 	TSharedPtr<FInternetAddr> OrigAddr;
 
@@ -447,23 +484,12 @@ void UIpConnection::LowLevelSend(void* Data, int32 CountBits, FOutPacketTraits&
 			RemoteAddr = OrigAddr;
 		}
 	};
-
-	if (GNetBlockSend)
-	{
-		return;
-	}
 #endif
 
-	if (Resolver->ShouldBlockSend())
-	{
-		UE_LOG(LogNet, Verbose, TEXT("Skipping send task as we are waiting on the next resolution step"));
-		return;
-	}
-
-	// Process any packet modifiers
+	// Process any packet modifiers, such as compression and encryption.
 	if (Handler.IsValid() && !Handler->GetRawSend())
 	{
-		const ProcessedPacket ProcessedData = Handler->Outgoing(reinterpret_cast<uint8*>(Data), CountBits, Traits);
+		const ProcessedPacket ProcessedData = Handler->Outgoing(DataToSend, CountBits, Traits);
 
 		if (!ProcessedData.bError)
 		{
@@ -480,7 +506,7 @@ void UIpConnection::LowLevelSend(void* Data, int32 CountBits, FOutPacketTraits&
 	int32 CountBytes = FMath::DivideAndRoundUp(CountBits, 8);
 
 #if !UE_BUILD_SHIPPING
-	LowLevelSendDel.ExecuteIfBound((void*)DataToSend, CountBytes, bBlockSend);
+	LowLevelSendDel.ExecuteIfBound(DataToSend, CountBytes, bBlockSend);
 #endif
 
 	if (!bBlockSend)
@@ -489,12 +515,12 @@ void UIpConnection::LowLevelSend(void* Data, int32 CountBits, FOutPacketTraits&
 		FSocketSendResult SendResult;
 		CLOCK_CYCLES(Driver->SendCycles);
 
-		if ( CountBytes > MaxPacket )
+		if (CountBytes > MaxPacket)
 		{
-			UE_LOG( LogNet, Warning, TEXT( "UIpConnection::LowLevelSend: CountBytes > MaxPacketSize! Count: %i, MaxPacket: %i %s" ), CountBytes, MaxPacket, *Describe() );
+			UE_LOG(LogNet, Warning, TEXT("UIpConnection::LowLevelSend: CountBytes > MaxPacketSize! Count: %i, MaxPacket: %i %s"), CountBytes, MaxPacket, *Describe());
 		}
 
-		FPacketAudit::NotifyLowLevelSend((uint8*)DataToSend, CountBytes, CountBits);
+		FPacketAudit::NotifyLowLevelSend(DataToSend, CountBytes, CountBits);
 
 		if (CountBytes > 0)
 		{
