[AutoRTFM] Ensure that packets sent in a transaction are always queued via OnCommit.

Previously, we would only handle packets sent in the closed, and let data through in the open. This ran the risk of allowing packets to be sent out of sequence. We will now always queue up packets when in a transaction.

diff --git a/Engine/Plugins/Online/OnlineSubsystemUtils/Source/OnlineSubsystemUtils/Private/IpConnection.cpp b/Engine/Plugins/Online/OnlineSubsystemUtils/Source/OnlineSubsystemUtils/Private/IpConnection.cpp
--- a/Engine/Plugins/Online/OnlineSubsystemUtils/Source/OnlineSubsystemUtils/Private/IpConnection.cpp
+++ b/Engine/Plugins/Online/OnlineSubsystemUtils/Source/OnlineSubsystemUtils/Private/IpConnection.cpp
@@ -439,19 +439,26 @@ void UIpConnection::LowLevelSend(void* Data, int32 CountBits, FOutPacketTraits&
 		return;
 	}
 
-	if (AutoRTFM::IsClosed())
+	if (AutoRTFM::IsTransactional())
 	{
 		// Inside an AutoRTFM transaction, we must postpone network sends until commit time, to account for the
-		// possibility of aborting the transaction. An abort will prevent any network traffic from occurring.
-		int32 CountBytes = FMath::DivideAndRoundUp(CountBits, 8);
-
-		AutoRTFM::OnCommit([this, 
-			                DataCopy = TArray<uint8>(DataToSend, CountBytes),
-							CountBits,
-			                TraitsCopy = Traits]() mutable
+		// possibility of aborting the transaction. (An abort will prevent any network traffic from occurring.)
+		// We aren't allowed to send packets out of order, so we explicitly jump into the closed here; this ensures
+		// that OnCommit takes effect for _all_ transactional packets, not just packets sent in the closed.
+		const AutoRTFM::EContextStatus Status = AutoRTFM::Close([&]
 		{
-			SendToRemote(DataCopy.GetData(), CountBits, TraitsCopy);
+			int32 CountBytes = FMath::DivideAndRoundUp(CountBits, 8);
+
+			AutoRTFM::OnCommit([this, 
+								DataCopy = TArray<uint8>(DataToSend, CountBytes),
+								CountBits,
+								TraitsCopy = Traits]() mutable
+			{
+				SendToRemote(DataCopy.GetData(), CountBits, TraitsCopy);
+			});
 		});
+
+		check(Status == AutoRTFM::EContextStatus::OnTrack);
 	}
 	else
 	{
