Renamed FPIELoginSettingsInternal to FOnlineAccountStoredCredentials to allow it to be reused for other online account editor config settings

#rnx

diff --git a/Engine/Plugins/Online/OnlineSubsystemUtils/Source/OnlineSubsystemUtils/Private/OnlinePIESettings.cpp b/Engine/Plugins/Online/OnlineSubsystemUtils/Source/OnlineSubsystemUtils/Private/OnlinePIESettings.cpp
--- a/Engine/Plugins/Online/OnlineSubsystemUtils/Source/OnlineSubsystemUtils/Private/OnlinePIESettings.cpp
+++ b/Engine/Plugins/Online/OnlineSubsystemUtils/Source/OnlineSubsystemUtils/Private/OnlinePIESettings.cpp
@@ -7,7 +7,7 @@
 #include UE_INLINE_GENERATED_CPP_BY_NAME(OnlinePIESettings)
 
 const int32 ONLINEPIE_XOR_KEY = 0xdeadbeef;
-void FPIELoginSettingsInternal::Encrypt()
+void FOnlineAccountStoredCredentials::Encrypt()
 {
 	if (Token.Len() > 0)
 	{
@@ -41,7 +41,7 @@ void FPIELoginSettingsInternal::Encrypt()
 	}
 }
 
-void FPIELoginSettingsInternal::Decrypt()
+void FOnlineAccountStoredCredentials::Decrypt()
 {
 	if (TokenBytes.Num() > 0)
 	{
@@ -96,7 +96,7 @@ void UOnlinePIESettings::PostInitProperties()
 {
 	Super::PostInitProperties();
 
-	for (FPIELoginSettingsInternal& Login : Logins)
+	for (FOnlineAccountStoredCredentials& Login : Logins)
 	{
 		Login.Decrypt();
 	}
@@ -122,10 +122,10 @@ void UOnlinePIESettings::PostEditChangeProperty(FPropertyChangedEvent& PropertyC
 			// If we paste on top of the whole login entry, all fields will have changed and need their checks run.
 			const bool bPastedAllValues = SubPropName == GET_MEMBER_NAME_CHECKED(UOnlinePIESettings, Logins);
 			
-			if (bPastedAllValues || SubPropName == GET_MEMBER_NAME_CHECKED(FPIELoginSettingsInternal, Id))
+			if (bPastedAllValues || SubPropName == GET_MEMBER_NAME_CHECKED(FOnlineAccountStoredCredentials, Id))
 			{
 				TSet<FString> Ids;
-				for (FPIELoginSettingsInternal& Login : Logins)
+				for (FOnlineAccountStoredCredentials& Login : Logins)
 				{
 					// Remove any whitespace from login input
 					Login.Id.TrimStartAndEndInline();
@@ -145,9 +145,9 @@ void UOnlinePIESettings::PostEditChangeProperty(FPropertyChangedEvent& PropertyC
 				}
 			}
 
-			if (bPastedAllValues || SubPropName == GET_MEMBER_NAME_CHECKED(FPIELoginSettingsInternal, Token))
+			if (bPastedAllValues || SubPropName == GET_MEMBER_NAME_CHECKED(FOnlineAccountStoredCredentials, Token))
 			{
-				for (FPIELoginSettingsInternal& Login : Logins)
+				for (FOnlineAccountStoredCredentials& Login : Logins)
 				{
 					// Remove any whitespace from login input
 					Login.Token.TrimStartAndEndInline();
@@ -156,9 +156,9 @@ void UOnlinePIESettings::PostEditChangeProperty(FPropertyChangedEvent& PropertyC
 				}
 			}
 
-			if (bPastedAllValues || SubPropName == GET_MEMBER_NAME_CHECKED(FPIELoginSettingsInternal, Type))
+			if (bPastedAllValues || SubPropName == GET_MEMBER_NAME_CHECKED(FOnlineAccountStoredCredentials, Type))
 			{
-				for (FPIELoginSettingsInternal& Login : Logins)
+				for (FOnlineAccountStoredCredentials& Login : Logins)
 				{
 					// Remove any whitespace from login input
 					Login.Type.TrimStartAndEndInline();
diff --git a/Engine/Plugins/Online/OnlineSubsystemUtils/Source/OnlineSubsystemUtils/Private/OnlinePIESettings.h b/Engine/Plugins/Online/OnlineSubsystemUtils/Source/OnlineSubsystemUtils/Private/OnlinePIESettings.h
--- a/Engine/Plugins/Online/OnlineSubsystemUtils/Source/OnlineSubsystemUtils/Private/OnlinePIESettings.h
+++ b/Engine/Plugins/Online/OnlineSubsystemUtils/Source/OnlineSubsystemUtils/Private/OnlinePIESettings.h
@@ -2,46 +2,11 @@
 #pragma once
 
 #include "Engine/DeveloperSettings.h"
+#include "OnlineAccountStoredCredentials.h"
 #include "OnlinePIESettings.generated.h"
 
 struct FPropertyChangedEvent;
 
-/**
- * Stores PIE login credentials
- */
-USTRUCT()
-struct FPIELoginSettingsInternal
-{
-public:
-
-	GENERATED_USTRUCT_BODY()
-
-	/** Id of the user logging in (email, display name, facebook id, etc) */
-	UPROPERTY(EditAnywhere, Category = "Logins", meta = (DisplayName = "User Id", Tooltip = "Id of the user logging in (email, display name, facebook id, etc)"))
-	FString Id;
-	/** Credentials of the user logging in (password or auth token) */
-	UPROPERTY(EditAnywhere, Transient, Category = "Logins", meta = (DisplayName = "Password", Tooltip = "Credentials of the user logging in (password or auth token)", PasswordField = true))
-	FString Token;
-	/** Type of account. Needed to identity the auth method to use (epic, internal, facebook, etc) */
-	UPROPERTY(EditAnywhere, Category = "Logins", meta = (DisplayName = "Type", Tooltip = "Type of account. Needed to identity the auth method to use (epic, internal, facebook, etc)"))
-	FString Type;
-	/** Token stored as an array of bytes, encrypted */
-	UPROPERTY()
-	TArray<uint8> TokenBytes;
-
-	/** @return true if the credentials are valid, false otherwise */
-	bool IsValid() const { return !Id.IsEmpty() && !Token.IsEmpty() && !Type.IsEmpty(); }
-
-	/**
-	 * Encrypt the Token field into the TokenBytes field
-	 */
-	void Encrypt();
-	/**
-	 * Decrypt the TokenBytes field into the Token field
-	 */
-	void Decrypt();
-};
-
 /**
  * Setup login credentials for the Play In Editor (PIE) feature
  */
@@ -59,7 +24,7 @@ public:
 
 	/** Array of credentials to use, one for each Play In Editor (PIE) instance */
 	UPROPERTY(config, EditAnywhere, Category = "Logins", meta = (DisplayName = "Credentials", Tooltip = "Login credentials, at least one for each instance of PIE that is intended to be run"))
-	TArray<FPIELoginSettingsInternal> Logins;
+	TArray<FOnlineAccountStoredCredentials> Logins;
 
 	// Begin UObject Interface
 	virtual void PostInitProperties() override;
diff --git a/Engine/Plugins/Online/OnlineSubsystemUtils/Source/OnlineSubsystemUtils/Private/OnlineServicesEngineUtilsImpl.cpp b/Engine/Plugins/Online/OnlineSubsystemUtils/Source/OnlineSubsystemUtils/Private/OnlineServicesEngineUtilsImpl.cpp
--- a/Engine/Plugins/Online/OnlineSubsystemUtils/Source/OnlineSubsystemUtils/Private/OnlineServicesEngineUtilsImpl.cpp
+++ b/Engine/Plugins/Online/OnlineSubsystemUtils/Source/OnlineSubsystemUtils/Private/OnlineServicesEngineUtilsImpl.cpp
@@ -153,7 +153,7 @@ int32 FOnlineServicesEngineUtils::GetNumPIELogins() const
 
 	int32 NumValidLogins = 0;
 	const UOnlinePIESettings* OnlinePIESettings = GetDefault<UOnlinePIESettings>();
-	for (const FPIELoginSettingsInternal& Login : OnlinePIESettings->Logins)
+	for (const FOnlineAccountStoredCredentials& Login : OnlinePIESettings->Logins)
 	{
 		if (Login.IsValid())
 		{
@@ -172,7 +172,7 @@ void FOnlineServicesEngineUtils::GetPIELogins(TArray<FAuthLogin::Params>& Logins
 {
 	check(UObjectInitialized());
 	const UOnlinePIESettings* OnlinePIESettings = GetDefault<UOnlinePIESettings>();
-	for (const FPIELoginSettingsInternal& PIELoginCredentials : OnlinePIESettings->Logins)
+	for (const FOnlineAccountStoredCredentials& PIELoginCredentials : OnlinePIESettings->Logins)
 	{
 		UE::Online::FAuthLogin::Params LoginParameters;
 		LoginParameters.CredentialsId = PIELoginCredentials.Id;
diff --git a/Engine/Plugins/Online/OnlineSubsystemUtils/Source/OnlineSubsystemUtils/Private/OnlineSubsystemUtilsModule.cpp b/Engine/Plugins/Online/OnlineSubsystemUtils/Source/OnlineSubsystemUtils/Private/OnlineSubsystemUtilsModule.cpp
--- a/Engine/Plugins/Online/OnlineSubsystemUtils/Source/OnlineSubsystemUtils/Private/OnlineSubsystemUtilsModule.cpp
+++ b/Engine/Plugins/Online/OnlineSubsystemUtils/Source/OnlineSubsystemUtils/Private/OnlineSubsystemUtilsModule.cpp
@@ -164,7 +164,7 @@ public:
 
 		int32 NumValidLogins = 0;
 		const UOnlinePIESettings* OnlinePIESettings = GetDefault<UOnlinePIESettings>();
-		for (const FPIELoginSettingsInternal& Login : OnlinePIESettings->Logins)
+		for (const FOnlineAccountStoredCredentials& Login : OnlinePIESettings->Logins)
 		{
 			if (Login.IsValid())
 			{
@@ -186,7 +186,7 @@ public:
 		if (OnlinePIESettings->Logins.Num() > 0)
 		{
 			Logins.Empty(OnlinePIESettings->Logins.Num());
-			for (const FPIELoginSettingsInternal& Login : OnlinePIESettings->Logins)
+			for (const FOnlineAccountStoredCredentials& Login : OnlinePIESettings->Logins)
 			{
 				if (Login.IsValid())
 				{
diff --git a/Engine/Plugins/Online/OnlineSubsystemUtils/Source/OnlineSubsystemUtils/Public/OnlineAccountStoredCredentials.h b/Engine/Plugins/Online/OnlineSubsystemUtils/Source/OnlineSubsystemUtils/Public/OnlineAccountStoredCredentials.h
new file mode 100644
--- /dev/null
+++ b/Engine/Plugins/Online/OnlineSubsystemUtils/Source/OnlineSubsystemUtils/Public/OnlineAccountStoredCredentials.h
@@ -0,0 +1,48 @@
+// Copyright Epic Games, Inc. All Rights Reserved.
+#pragma once
+
+#include "Containers/Array.h"
+#include "Containers/UnrealString.h"
+#include "OnlineAccountStoredCredentials.generated.h"
+
+/**
+ * Stores online account login credentials (for editor config data).
+ * @note ONLY use this in trusted environments (like a local config file) and NOT for anything that requires actual security/strong encryption.
+ * @see FOnlineAccountCredentials.
+ */
+USTRUCT()
+struct FOnlineAccountStoredCredentials
+{
+public:
+
+	GENERATED_BODY()
+
+	/** Id of the user logging in (email, display name, facebook id, etc) */
+	UPROPERTY(EditAnywhere, Category = "Logins", meta = (DisplayName = "User Id", Tooltip = "Id of the user logging in (email, display name, facebook id, etc)"))
+	FString Id;
+	/** Credentials of the user logging in (password or auth token) */
+	UPROPERTY(EditAnywhere, Transient, Category = "Logins", meta = (DisplayName = "Password", Tooltip = "Credentials of the user logging in (password or auth token)", PasswordField = true))
+	FString Token;
+	/** Type of account. Needed to identity the auth method to use (epic, internal, facebook, etc) */
+	UPROPERTY(EditAnywhere, Category = "Logins", meta = (DisplayName = "Type", Tooltip = "Type of account. Needed to identity the auth method to use (epic, internal, facebook, etc)"))
+	FString Type;
+	/** Token stored as an array of bytes, encrypted */
+	UPROPERTY()
+	TArray<uint8> TokenBytes;
+
+	/** @return true if the credentials are valid, false otherwise */
+	bool IsValid() const
+	{
+		return !Id.IsEmpty() && !Token.IsEmpty() && !Type.IsEmpty();
+	}
+
+	/**
+	 * Encrypt the Token field into the TokenBytes field
+	 */
+	ONLINESUBSYSTEMUTILS_API void Encrypt();
+	
+	/**
+	 * Decrypt the TokenBytes field into the Token field
+	 */
+	ONLINESUBSYSTEMUTILS_API void Decrypt();
+};
