Remove Online logic marked for deprecation in 5.1 and 5.0

diff --git a/Engine/Plugins/Online/OnlineSubsystemSteam/Source/Private/OnlineSessionAsyncServerSteam.cpp b/Engine/Plugins/Online/OnlineSubsystemSteam/Source/Private/OnlineSessionAsyncServerSteam.cpp
--- a/Engine/Plugins/Online/OnlineSubsystemSteam/Source/Private/OnlineSessionAsyncServerSteam.cpp
+++ b/Engine/Plugins/Online/OnlineSubsystemSteam/Source/Private/OnlineSessionAsyncServerSteam.cpp
@@ -25,61 +25,15 @@
  * matchmaking will not work for your dedicated server
  */
 #ifndef UE_PROJECT_STEAMPRODUCTNAME
-
-#ifdef STEAMPRODUCTNAME
-UE_DEPRECATED(4.22, "The Steam Product Name Macro has been updated to be configurable from your Target.cs file. Please change STEAMPRODUCTNAME to UE_PROJECT_STEAMPRODUCTNAME.")
-#define UE_PROJECT_STEAMPRODUCTNAME STEAMPRODUCTNAME
-#else
-
-#ifdef UE4_PROJECT_STEAMPRODUCTNAME
-UE_DEPRECATED(5.0, "UE4_PROJECT_STEAMPRODUCTNAME has been renamed to UE_PROJECT_STEAMPRODUCTNAME.")
-#define UE_PROJECT_STEAMPRODUCTNAME UE4_PROJECT_STEAMPRODUCTNAME
-#else
 #define UE_PROJECT_STEAMPRODUCTNAME "unrealdk"
-#endif //UE4_PROJECT_STEAMPRODUCTNAME
-
-#endif //STEAMPRODUCTNAME
-
 #endif //UE_PROJECT_STEAMPRODUCTNAME
 
-
-
 #ifndef UE_PROJECT_STEAMGAMEDIR
-
-#ifdef STEAMGAMEDIR
-UE_DEPRECATED(4.22, "The Steam Game Directory Macro has been updated to be configurable from your Target.cs file. Please change STEAMGAMEDIR to UE_PROJECT_STEAMGAMEDIR.")
-#define UE_PROJECT_STEAMGAMEDIR STEAMGAMEDIR
-#else
-
-#ifdef UE4_PROJECT_STEAMGAMEDIR
-UE_DEPRECATED(5.0, "UE4_PROJECT_STEAMGAMEDIR has been renamed to UE_PROJECT_STEAMGAMEDIR.")
-#define UE_PROJECT_STEAMGAMEDIR UE4_PROJECT_STEAMGAMEDIR
-#else
 #define UE_PROJECT_STEAMGAMEDIR "unrealtest"
-#endif // UE4_PROJECT_STEAMPRODUCTNAME
-
-#endif // STEAMGAMEDIR
-
 #endif // UE_PROJECT_STEAMGAMEDIR
 
-
-
 #ifndef UE_PROJECT_STEAMGAMEDESC
-
-#ifdef STEAMGAMEDESC
-UE_DEPRECATED(4.22, "The Steam Game Description Macro has been updated to be configurable from your Target.cs file. Please change STEAMGAMEDESC to UE_PROJECT_STEAMGAMEDESC.")
-#define UE_PROJECT_STEAMGAMEDESC STEAMGAMEDESC
-#else
-
-#ifdef UE4_PROJECT_STEAMGAMEDESC
-UE_DEPRECATED(5.0, "UE4_PROJECT_STEAMGAMEDESC has been renamed to UE_PROJECT_STEAMGAMEDESC.")
-#define UE_PROJECT_STEAMGAMEDESC UE4_PROJECT_STEAMGAMEDESC
-#else
 #define UE_PROJECT_STEAMGAMEDESC "Unreal Test!"
-#endif // UE4_PROJECT_STEAMGAMEDESC
-
-#endif // STEAMGAMEDESC
-
 #endif // UE_PROJECT_STEAMGAMEDESC
 
 /**
diff --git a/Engine/Plugins/Online/OnlineSubsystemSteam/Source/Private/OnlineSubsystemSteam.cpp b/Engine/Plugins/Online/OnlineSubsystemSteam/Source/Private/OnlineSubsystemSteam.cpp
--- a/Engine/Plugins/Online/OnlineSubsystemSteam/Source/Private/OnlineSubsystemSteam.cpp
+++ b/Engine/Plugins/Online/OnlineSubsystemSteam/Source/Private/OnlineSubsystemSteam.cpp
@@ -35,13 +35,8 @@
  * if your game is launched outside of Steam.
  */
 #ifndef UE_PROJECT_STEAMSHIPPINGID
-#ifdef UE4_PROJECT_STEAMSHIPPINGID
-UE_DEPRECATED(5.0, "UE4_PROJECT_STEAMSHIPPINGID has been renamed to UE_PROJECT_STEAMSHIPPINGID.")
-#define UE_PROJECT_STEAMSHIPPINGID UE4_PROJECT_STEAMSHIPPINGID
-#else
 #define UE_PROJECT_STEAMSHIPPINGID 0
 #endif
-#endif
 
 namespace FNetworkProtocolTypes
 {
diff --git a/Engine/Plugins/Online/OnlineSubsystemUtils/Source/OnlineSubsystemUtils/Classes/IpConnection.h b/Engine/Plugins/Online/OnlineSubsystemUtils/Source/OnlineSubsystemUtils/Classes/IpConnection.h
--- a/Engine/Plugins/Online/OnlineSubsystemUtils/Source/OnlineSubsystemUtils/Classes/IpConnection.h
+++ b/Engine/Plugins/Online/OnlineSubsystemUtils/Source/OnlineSubsystemUtils/Classes/IpConnection.h
@@ -20,20 +20,6 @@ namespace UE::Net::Private
 	class FNetConnectionAddressResolution;
 }
 
-
-/** A state system of the address resolution functionality. */
-enum class UE_DEPRECATED(5.1, "EAddressResolutionState has been moved to a private namespace.") EAddressResolutionState : uint8
-{
-	None = 0,
-	Disabled,
-	WaitingForResolves,
-	Connecting,
-	TryNextAddress,
-	Connected,
-	Done,	
-	Error
-};
-
 UCLASS(transient, config=Engine)
 class ONLINESUBSYSTEMUTILS_API UIpConnection : public UNetConnection
 {
@@ -42,12 +28,6 @@ class ONLINESUBSYSTEMUTILS_API UIpConnection : public UNetConnection
 	friend FIpConnectionHelper;
 	friend UE::Net::Private::FNetDriverAddressResolution;
 
-public:
-	/** This is a non-owning pointer to a socket owned elsewhere, IpConnection will not destroy the socket through this pointer. */
-	UE_DEPRECATED(5.1, "Socket access is now controlled through GetSocket")
-	FSocket*				Socket;
-
-
 public:
 	//~ Begin NetConnection Interface
 	virtual void InitBase(UNetDriver* InDriver, class FSocket* InSocket, const FURL& InURL, EConnectionState InState, int32 InMaxPacket = 0, int32 InPacketOverhead = 0) override;
@@ -75,13 +55,6 @@ public:
 	 */
 	FSocket* GetSocket() const
 	{
-		PRAGMA_DISABLE_DEPRECATION_WARNINGS
-		if (Socket != SocketPrivate.Get())
-		{
-			return Socket;
-		}
-		PRAGMA_ENABLE_DEPRECATION_WARNINGS
-
 		return SocketPrivate.Get();
 	}
 
@@ -148,9 +121,6 @@ private:
 	/** Sets the local socket pointer, and safely cleans up any references to old sockets */
 	void SetSocket_Local(const TSharedPtr<FSocket>& InSocket);
 
-	/** Cleanup for the deprecated 'Socket' value */
-	void CleanupDeprecatedSocket();
-
 	/** Safe/non-blocking cleanup of a shared socket, which may be in use by async sends, or may be at risk of triggering ICMP unreachable errors */
 	void SafeDeferredSocketCleanup(const TSharedPtr<FSocket>& InSocket);
 
diff --git a/Engine/Plugins/Online/OnlineSubsystemUtils/Source/OnlineSubsystemUtils/Classes/IpNetDriver.h b/Engine/Plugins/Online/OnlineSubsystemUtils/Source/OnlineSubsystemUtils/Classes/IpNetDriver.h
--- a/Engine/Plugins/Online/OnlineSubsystemUtils/Source/OnlineSubsystemUtils/Classes/IpNetDriver.h
+++ b/Engine/Plugins/Online/OnlineSubsystemUtils/Source/OnlineSubsystemUtils/Classes/IpNetDriver.h
@@ -129,15 +129,6 @@ public:
 	 */
 	virtual FSocket* GetSocket();
 
-	/**
-	* Set the current NetDriver's Socket/LocalAddr to the given socket (typically after Net Address Resolution).
-	* This will automatically pass the socket back through all NetConnection's, and trigger safe/deferred cleanup of the old socket.
-	*
-	* @param NewSocket	The socket pointer to set this NetDriver's socket to
-	*/
-	UE_DEPRECATED(5.1, "Use the version of SetSocketAndLocalAddress which takes a TSharedPtr")
-	void SetSocketAndLocalAddress(FSocket* NewSocket);
-
 	/**
 	* Set the current NetDriver's Socket/LocalAddr to the given socket (typically after Net Address Resolution).
 	* This will automatically pass the socket back through all NetConnection's, and trigger safe/deferred cleanup of the old socket.
diff --git a/Engine/Plugins/Online/OnlineSubsystemUtils/Source/OnlineSubsystemUtils/Private/IpConnection.cpp b/Engine/Plugins/Online/OnlineSubsystemUtils/Source/OnlineSubsystemUtils/Private/IpConnection.cpp
--- a/Engine/Plugins/Online/OnlineSubsystemUtils/Source/OnlineSubsystemUtils/Private/IpConnection.cpp
+++ b/Engine/Plugins/Online/OnlineSubsystemUtils/Source/OnlineSubsystemUtils/Private/IpConnection.cpp
@@ -95,15 +95,10 @@ namespace UE::Net::Private
 
 UIpConnection::UIpConnection(const FObjectInitializer& ObjectInitializer) :
 	Super(ObjectInitializer),
-	PRAGMA_DISABLE_DEPRECATION_WARNINGS
-	Socket(nullptr),
-	PRAGMA_ENABLE_DEPRECATION_WARNINGS
 	SocketErrorDisconnectDelay(5.f),
 	SocketError_SendDelayStartTime(0.f),
 	SocketError_RecvDelayStartTime(0.f),
-	PRAGMA_DISABLE_DEPRECATION_WARNINGS
-	Resolver(MakePimpl<UE::Net::Private::FNetConnectionAddressResolution>(Socket))
-	PRAGMA_ENABLE_DEPRECATION_WARNINGS
+	Resolver(MakePimpl<UE::Net::Private::FNetConnectionAddressResolution>())
 {
 	// Auto add address resolution disable flags if the cvar is set.
 	if (!!CVarNetIpConnectionDisableResolution.GetValueOnAnyThread())
@@ -122,9 +117,7 @@ void UIpConnection::InitBase(UNetDriver* InDriver, class FSocket* InSocket, cons
 
 	if (!Resolver->IsAddressResolutionEnabled())
 	{
-		PRAGMA_DISABLE_DEPRECATION_WARNINGS
-		Socket = InSocket;
-		PRAGMA_ENABLE_DEPRECATION_WARNINGS
+		SocketPrivate = InSocket->AsShared();
 	}
 
 	if (CVarNetEnableCongestionControl.GetValueOnAnyThread() > 0)
@@ -244,10 +237,6 @@ void UIpConnection::Tick(float DeltaSeconds)
 		SetConnectionState(USOCK_Closed);
 		Close(ENetCloseResult::AddressResolutionFailed);
 	}
-	else if (CheckResult == ECheckAddressResolutionResult::FindSocketError)
-	{
-		CleanupDeprecatedSocket();
-	}
 
 	// Deferred cleanup of old sockets also waits until packets are received on the new socket (which happens after restart handshake process),
 	// in order to avoid getting kicked by the server, due to closing the old socket too early and triggering ICMP unreachable errors
@@ -314,7 +303,6 @@ void UIpConnection::CleanUp()
 	WaitForSendTasks();
 
 	// Sockets must be cleaned up after send tasks complete, as they can still be in use
-	CleanupDeprecatedSocket();
 	Resolver->CleanupResolutionSockets();
 
 	DeferredCleanupReadyCount = 0;
@@ -330,7 +318,6 @@ void UIpConnection::SetSocket_Local(const TSharedPtr<FSocket>& InSocket)
 	if (ResolutionSocket.IsValid() && InSocket != ResolutionSocket)
 	{
 		SafeDeferredSocketCleanup(ResolutionSocket);
-		CleanupDeprecatedSocket();
 		Resolver->CleanupResolutionSockets();
 	}
 
@@ -339,23 +326,9 @@ void UIpConnection::SetSocket_Local(const TSharedPtr<FSocket>& InSocket)
 		SafeDeferredSocketCleanup(SocketPrivate);
 	}
 
-	PRAGMA_DISABLE_DEPRECATION_WARNINGS
-	Socket = InSocket.Get();
-	PRAGMA_ENABLE_DEPRECATION_WARNINGS
-
 	SocketPrivate = InSocket;
 }
 
-void UIpConnection::CleanupDeprecatedSocket()
-{
-	PRAGMA_DISABLE_DEPRECATION_WARNINGS
-	if (Socket != nullptr && Resolver->IsAddressResolutionEnabled())
-	{
-		Socket = nullptr;
-	}
-	PRAGMA_ENABLE_DEPRECATION_WARNINGS
-}
-
 // This needs to be removed when the main multithreading changes are added (or implemented in line with them)
 void UIpConnection::SafeDeferredSocketCleanup(const TSharedPtr<FSocket>& InSocket)
 {
diff --git a/Engine/Plugins/Online/OnlineSubsystemUtils/Source/OnlineSubsystemUtils/Private/IpNetDriver.cpp b/Engine/Plugins/Online/OnlineSubsystemUtils/Source/OnlineSubsystemUtils/Private/IpNetDriver.cpp
--- a/Engine/Plugins/Online/OnlineSubsystemUtils/Source/OnlineSubsystemUtils/Private/IpNetDriver.cpp
+++ b/Engine/Plugins/Online/OnlineSubsystemUtils/Source/OnlineSubsystemUtils/Private/IpNetDriver.cpp
@@ -724,11 +724,6 @@ private:
 		Connection->HandleSocketRecvError(Driver, ErrorString);
 	}
 
-	static void CleanupDeprecatedSocket(UIpConnection* Connection)
-	{
-		Connection->CleanupDeprecatedSocket();
-	}
-
 	static void SetSocket_Local(UIpConnection* Connection, const TSharedPtr<FSocket>& InSocket)
 	{
 		Connection->SetSocket_Local(InSocket);
@@ -865,11 +860,6 @@ FUniqueSocket UIpNetDriver::CreateAndBindSocket(TSharedRef<FInternetAddr> BindAd
 	return NewSocket;
 }
 
-void UIpNetDriver::SetSocketAndLocalAddress(FSocket* NewSocket)
-{
-	SetSocketAndLocalAddress(TSharedPtr<FSocket>(NewSocket, FSocketDeleter(GetSocketSubsystem())));
-}
-
 void UIpNetDriver::SetSocketAndLocalAddress(const TSharedPtr<FSocket>& SharedSocket)
 {
 	// Must be called even if the current socket is already SharedSocket (for when Net Address Resolution resolves the current socket)
@@ -1607,7 +1597,6 @@ void UIpNetDriver::LowLevelDestroy()
 
 		if (ConnResolver != nullptr && ConnResolver->IsAddressResolutionEnabled())
 		{
-			FIpConnectionHelper::CleanupDeprecatedSocket(IpServerConnection);
 			ConnResolver->CleanupResolutionSockets();
 		}
 
diff --git a/Engine/Plugins/Online/OnlineSubsystemUtils/Source/OnlineSubsystemUtils/Private/NetAddressResolution.cpp b/Engine/Plugins/Online/OnlineSubsystemUtils/Source/OnlineSubsystemUtils/Private/NetAddressResolution.cpp
--- a/Engine/Plugins/Online/OnlineSubsystemUtils/Source/OnlineSubsystemUtils/Private/NetAddressResolution.cpp
+++ b/Engine/Plugins/Online/OnlineSubsystemUtils/Source/OnlineSubsystemUtils/Private/NetAddressResolution.cpp
@@ -180,12 +180,6 @@ FNetConnectionAddressResolution* FNetDriverAddressResolution::GetConnectionResol
  * FNetConnectionAddressResolution
  */
 
-FNetConnectionAddressResolution::FNetConnectionAddressResolution(const FSocket* const & InDeprecatedSocket)
-	: DeprecatedSocket(InDeprecatedSocket)
-	, ResolutionState(EAddressResolutionState::None)
-{
-}
-
 bool FNetConnectionAddressResolution::InitLocalConnection(ISocketSubsystem* SocketSubsystem, FSocket* InSocket, const FURL& InURL)
 {
 	bool bValidInit = true;
@@ -368,15 +362,6 @@ void FNetConnectionAddressResolution::CleanupResolutionSockets(ECleanupResolutio
 	{
 		if (EnumHasAnyFlags(CleanupFlags, ECleanupResolutionSocketsFlags::CleanAll))
 		{
-			// Remove when UIpConnection.Socket is deprecated
-#if !UE_BUILD_SHIPPING && !PLATFORM_ANDROID
-			if ( (!GIsBuildMachine) && (!IsRunningCommandlet()) ) // make sure we don't hit the ensure while on the Android testing... 
-			{
-				ensureMsgf(DeprecatedSocket == nullptr,
-					TEXT("All codepaths passing 'CleanAll', must call 'UIpConnection::CleanupDeprecatedSocket' first."));
-			}
-#endif
-
 			ResolutionSocket.Reset();
 		}
 	
diff --git a/Engine/Plugins/Online/OnlineSubsystemUtils/Source/OnlineSubsystemUtils/Private/NetAddressResolution.h b/Engine/Plugins/Online/OnlineSubsystemUtils/Source/OnlineSubsystemUtils/Private/NetAddressResolution.h
--- a/Engine/Plugins/Online/OnlineSubsystemUtils/Source/OnlineSubsystemUtils/Private/NetAddressResolution.h
+++ b/Engine/Plugins/Online/OnlineSubsystemUtils/Source/OnlineSubsystemUtils/Private/NetAddressResolution.h
@@ -172,13 +172,6 @@ class FNetConnectionAddressResolution
 	friend FNetDriverAddressResolution;
 
 public:
-	/**
-	 * Base constructor
-	 *
-	 * @param InDeprecatedSocket	Read-only reference to deprecated UIpConnection.Socket variable, for verification checks
-	 */
-	FNetConnectionAddressResolution(const FSocket* const & InDeprecatedSocket);
-
 	/**
 	 * Checks to see if the owner NetConnection class can use address resolution
 	 *
@@ -321,9 +314,6 @@ private:
 
 
 private:
-	/** Temporary read-only reference to the deprecated UIpConnection.Socket value, for verification checks */
-	const FSocket* const& DeprecatedSocket;
-
 	/** An array of sockets tied to every binding address. */
 	TArray<TSharedPtr<FSocket>> BindSockets;
 
@@ -343,6 +333,6 @@ private:
 	 *  The connection's current status of where it is in the resolution state machine.
 	 *  If a platform should not use resolution, call DisableAddressResolution() in your constructor
 	 */
-	EAddressResolutionState ResolutionState;
+	EAddressResolutionState ResolutionState = EAddressResolutionState::None;
 };
 }
diff --git a/Engine/Source/Runtime/Sockets/Public/Sockets.h b/Engine/Source/Runtime/Sockets/Public/Sockets.h
--- a/Engine/Source/Runtime/Sockets/Public/Sockets.h
+++ b/Engine/Source/Runtime/Sockets/Public/Sockets.h
@@ -15,7 +15,7 @@ class FInternetAddr;
 /**
  * This is our abstract base class that hides the platform specific socket implementation
  */
-class FSocket
+class FSocket : public TSharedFromThis<FSocket, ESPMode::ThreadSafe>
 {
 protected:
 
