Added additional hardening for matchmaking

Updated ReservationResult enum to account for reservations being denied by the game specific validation. Added more validation to the client when processing party reservation changes.

#tests Ran through a variety of tests involving joining in progress and rejoining games, was able to connect as expected. Spoofed rejoin request to attempt to join a private game and the player wasn't allowed in.

diff --git a/Engine/Plugins/Online/OnlineSubsystemUtils/Source/OnlineSubsystemUtils/Private/PartyBeaconHost.cpp b/Engine/Plugins/Online/OnlineSubsystemUtils/Source/OnlineSubsystemUtils/Private/PartyBeaconHost.cpp
--- a/Engine/Plugins/Online/OnlineSubsystemUtils/Source/OnlineSubsystemUtils/Private/PartyBeaconHost.cpp
+++ b/Engine/Plugins/Online/OnlineSubsystemUtils/Source/OnlineSubsystemUtils/Private/PartyBeaconHost.cpp
@@ -714,7 +714,7 @@ EPartyReservationResult::Type APartyBeaconHost::AddPartyReservation(const FParty
 								}
 								else
 								{
-									Result = EPartyReservationResult::ReservationDenied;
+									Result = EPartyReservationResult::ReservationDenied_ValidationFailed;
 								}
 							}
 							else
@@ -1137,7 +1137,7 @@ void APartyBeaconHost::ProcessReservationRequest(APartyBeaconClient* Client, con
 		ReservationRequest.Dump();
 	}
 
-	if (Client)
+	if (Client && Client->GetUniqueId().IsValid() && Client->GetUniqueId() == ReservationRequest.PartyLeader)
 	{
 		EPartyReservationResult::Type Result = EPartyReservationResult::BadSessionId;
 		if (DoesSessionMatch(SessionId))
@@ -1154,6 +1154,10 @@ void APartyBeaconHost::ProcessReservationRequest(APartyBeaconClient* Client, con
 
 		Client->ClientReservationResponse(Result);
 	}
+	else if (Client)
+	{
+		Client->ClientCancelReservationResponse(EPartyReservationResult::ReservationInvalid);
+	}
 }
 
 void APartyBeaconHost::ProcessReservationUpdateRequest(APartyBeaconClient* Client, const FString& SessionId, const FPartyReservation& ReservationUpdateRequest, bool bIsRemovingMember)
@@ -1166,7 +1170,7 @@ void APartyBeaconHost::ProcessReservationUpdateRequest(APartyBeaconClient* Clien
 		ReservationUpdateRequest.PartyMembers.Num(),
 		Client ? *Client->GetNetConnection()->LowLevelDescribe() : TEXT("NULL"));
 
-	if (Client)
+	if (Client && Client->GetUniqueId().IsValid() && Client->GetUniqueId() == ReservationUpdateRequest.PartyLeader)
 	{
 		EPartyReservationResult::Type Result = EPartyReservationResult::BadSessionId;
 		if (DoesSessionMatch(SessionId))
@@ -1184,6 +1188,10 @@ void APartyBeaconHost::ProcessReservationUpdateRequest(APartyBeaconClient* Clien
 
 		Client->ClientReservationResponse(Result);
 	}
+	else if(Client)
+	{
+		Client->ClientCancelReservationResponse(EPartyReservationResult::ReservationInvalid);
+	}
 }
 
 void APartyBeaconHost::ProcessReservationAddOrUpdateRequest(APartyBeaconClient* Client, const FString& SessionId, const FPartyReservation& ReservationRequest)
@@ -1195,7 +1203,7 @@ void APartyBeaconHost::ProcessReservationAddOrUpdateRequest(APartyBeaconClient*
 		ReservationRequest.PartyMembers.Num(),
 		Client ? *Client->GetNetConnection()->LowLevelDescribe() : TEXT("NULL"));
 
-	if (Client)
+	if (Client && Client->GetUniqueId().IsValid() && Client->GetUniqueId() == ReservationRequest.PartyLeader)
 	{
 		EPartyReservationResult::Type Result = EPartyReservationResult::BadSessionId;
 		if (DoesSessionMatch(SessionId))
@@ -1221,6 +1229,10 @@ void APartyBeaconHost::ProcessReservationAddOrUpdateRequest(APartyBeaconClient*
 
 		Client->ClientReservationResponse(Result);
 	}
+	else if (Client)
+	{
+		Client->ClientCancelReservationResponse(EPartyReservationResult::ReservationInvalid);
+	}
 }
 
 bool APartyBeaconHost::HasCrossplayOptOutReservation() const
@@ -1250,7 +1262,7 @@ void APartyBeaconHost::ProcessCancelReservationRequest(APartyBeaconClient* Clien
 		PartyLeader.IsValid() ? *PartyLeader->ToString() : TEXT("INVALID"),
 		Client ? *Client->GetNetConnection()->LowLevelDescribe() : TEXT("NULL"));
 
-	if (Client)
+	if (Client && Client->GetUniqueId().IsValid() && Client->GetUniqueId() == PartyLeader)
 	{
 		EPartyReservationResult::Type Result = RemovePartyReservation(PartyLeader);
 		UE_LOG(LogPartyBeacon, Verbose, TEXT("ProcessCancelReservationRequest result: %s"), EPartyReservationResult::ToString(Result));
@@ -1262,6 +1274,10 @@ void APartyBeaconHost::ProcessCancelReservationRequest(APartyBeaconClient* Clien
 
 		Client->ClientCancelReservationResponse(Result);
 	}
+	else if (Client)
+	{
+		Client->ClientCancelReservationResponse(EPartyReservationResult::ReservationInvalid);
+	}
 }
 
 void APartyBeaconHost::DumpReservations() const
diff --git a/Engine/Plugins/Online/OnlineSubsystemUtils/Source/OnlineSubsystemUtils/Public/PartyBeaconState.h b/Engine/Plugins/Online/OnlineSubsystemUtils/Source/OnlineSubsystemUtils/Public/PartyBeaconState.h
--- a/Engine/Plugins/Online/OnlineSubsystemUtils/Source/OnlineSubsystemUtils/Public/PartyBeaconState.h
+++ b/Engine/Plugins/Online/OnlineSubsystemUtils/Source/OnlineSubsystemUtils/Public/PartyBeaconState.h
@@ -45,6 +45,8 @@ namespace EPartyReservationResult
 		BadSessionId,
 		/** The reservation contains players already in this game */
 		ReservationDenied_ContainsExistingPlayers,
+		/** The reservation was rejected because it failed additional, project-specific validation */
+		ReservationDenied_ValidationFailed
 	};
 }
 
