Change PIE handling for duplicate credentials to be config-driven

diff --git a/Engine/Plugins/Online/OnlineSubsystemUtils/Config/Editor.ini b/Engine/Plugins/Online/OnlineSubsystemUtils/Config/Editor.ini
new file mode 100644
--- /dev/null
+++ b/Engine/Plugins/Online/OnlineSubsystemUtils/Config/Editor.ini
@@ -0,0 +1,2 @@
+[/Script/OnlineSubsystemUtils.OnlinePIEConfig]
++LoginTypesAllowingDuplicates=Developer
\ No newline at end of file
diff --git a/Engine/Plugins/Online/OnlineSubsystemUtils/Source/OnlineSubsystemUtils/Private/OnlinePIEConfig.cpp b/Engine/Plugins/Online/OnlineSubsystemUtils/Source/OnlineSubsystemUtils/Private/OnlinePIEConfig.cpp
new file mode 100644
--- /dev/null
+++ b/Engine/Plugins/Online/OnlineSubsystemUtils/Source/OnlineSubsystemUtils/Private/OnlinePIEConfig.cpp
@@ -0,0 +1,10 @@
+// Copyright Epic Games, Inc. All Rights Reserved.
+
+#include "OnlinePIEConfig.h"
+
+#include UE_INLINE_GENERATED_CPP_BY_NAME(OnlinePIEConfig)
+
+UOnlinePIEConfig::UOnlinePIEConfig(const FObjectInitializer& ObjectInitializer)
+	: Super(ObjectInitializer)
+{
+}
diff --git a/Engine/Plugins/Online/OnlineSubsystemUtils/Source/OnlineSubsystemUtils/Private/OnlinePIEConfig.h b/Engine/Plugins/Online/OnlineSubsystemUtils/Source/OnlineSubsystemUtils/Private/OnlinePIEConfig.h
new file mode 100644
--- /dev/null
+++ b/Engine/Plugins/Online/OnlineSubsystemUtils/Source/OnlineSubsystemUtils/Private/OnlinePIEConfig.h
@@ -0,0 +1,19 @@
+// Copyright Epic Games, Inc. All Rights Reserved.
+#pragma once
+
+#include "UObject/Object.h"
+#include "OnlinePIEConfig.generated.h"
+
+struct FPropertyChangedEvent;
+
+UCLASS(config=Editor)
+class UOnlinePIEConfig : public UObject
+{
+	GENERATED_BODY()
+public:
+	UOnlinePIEConfig(const FObjectInitializer& ObjectInitializer);
+
+	/** Type types of login credentials which allow duplicate entries. */
+	UPROPERTY(config)
+	TArray<FString> LoginTypesAllowingDuplicates;
+};
diff --git a/Engine/Plugins/Online/OnlineSubsystemUtils/Source/OnlineSubsystemUtils/Private/OnlinePIESettings.cpp b/Engine/Plugins/Online/OnlineSubsystemUtils/Source/OnlineSubsystemUtils/Private/OnlinePIESettings.cpp
--- a/Engine/Plugins/Online/OnlineSubsystemUtils/Source/OnlineSubsystemUtils/Private/OnlinePIESettings.cpp
+++ b/Engine/Plugins/Online/OnlineSubsystemUtils/Source/OnlineSubsystemUtils/Private/OnlinePIESettings.cpp
@@ -1,6 +1,7 @@
 // Copyright Epic Games, Inc. All Rights Reserved.
 
 #include "OnlinePIESettings.h"
+#include "OnlinePIEConfig.h"
 #include "OnlineSubsystem.h"
 #include "Misc/AES.h"
 
@@ -131,8 +132,9 @@ void UOnlinePIESettings::PostEditChangeProperty(FPropertyChangedEvent& PropertyC
 					Login.Id.TrimStartAndEndInline();
 
 					const bool bIsExistingLogin = Ids.Contains(Login.Id);
-					// Allow duplicate for Developer credentials.  This is intended to support logins where other fields uniquely identify the login but require the same id.
-					const bool bDuplicateAllowed = Login.Type.Equals(TEXT("Developer"), ESearchCase::IgnoreCase);
+					// Allow duplicate credentials or when allowed.  This is intended to support logins where other fields uniquely identify the login but require the same id.
+					const UOnlinePIEConfig* OnlinePIEConfig = GetDefault<UOnlinePIEConfig>();
+					const bool bDuplicateAllowed = OnlinePIEConfig->LoginTypesAllowingDuplicates.Contains(Login.Type);
 					if (bIsExistingLogin && !bDuplicateAllowed)
 					{
 						// Don't allow duplicate login ids
