Remove auth methods 1 & 2 from OnlineBeaconUnitTests as they were removed from beacons in CL 36500636. Fixes AuthenticatedHandshakeSuccess and AuthenticatedHandshakeFailure test cases.

diff --git a/Engine/Plugins/Online/OnlineSubsystemUtils/Source/OnlineSubsystemUtils/Private/Tests/OnlineBeaconUnitTest.cpp b/Engine/Plugins/Online/OnlineSubsystemUtils/Source/OnlineSubsystemUtils/Private/Tests/OnlineBeaconUnitTest.cpp
--- a/Engine/Plugins/Online/OnlineSubsystemUtils/Source/OnlineSubsystemUtils/Private/Tests/OnlineBeaconUnitTest.cpp
+++ b/Engine/Plugins/Online/OnlineSubsystemUtils/Source/OnlineSubsystemUtils/Private/Tests/OnlineBeaconUnitTest.cpp
@@ -1559,403 +1559,6 @@ bool FOnlineBeaconTestAuthenticatedHandshakeSuccess::RunTest(const FString& Para
 	UTEST_TRUE_EXPR(BeaconUnitTest::ConfigureBeaconNetDriver(*Prerequisites, BeaconHost, &BeaconHostNetStats));
 	UTEST_TRUE_EXPR(BeaconHostNetStats.IsValid());
 
-	// Method 1.
-	// Host handshake calls StartVerifyAuthentication and expects OnAuthenticationVerificationComplete to be called to complete authentication.
-	// Test OnAuthenticationVerificationComplete fired with no delay is working.
-	{
-		BeaconHostNetStats->ReceivedControlMessages.Empty();
-		Prerequisites->GetConfig() = BaseConfig;
-		Prerequisites->GetStats() = BeaconUnitTest::FTestStats();
-
-		Prerequisites->GetConfig().Auth.Method1.bEnabled = true;
-		Prerequisites->GetConfig().Auth.Method1.Result = FOnlineError::Success();
-
-		TObjectPtr<AOnlineBeaconClient> HostUserBeaconClient;
-
-		// Client setup.
-
-		TObjectPtr<AOnlineBeaconUnitTestClient> BeaconClient = Prerequisites->GetWorld()->SpawnActor<AOnlineBeaconUnitTestClient>(AOnlineBeaconUnitTestClient::StaticClass());
-		ON_SCOPE_EXIT{ if (BeaconClient) { BeaconClient->DestroyBeacon(); } };
-		UTEST_TRUE_EXPR(BeaconUnitTest::ConfigureBeacon(*Prerequisites, BeaconClient));
-
-		UTEST_TRUE_EXPR(BeaconClient->GetConnectionState() == EBeaconConnectionState::Invalid);
-		UTEST_TRUE_EXPR(BeaconUnitTest::InitClientForUser(*Prerequisites, BeaconClient, UserId));
-		UTEST_TRUE_EXPR(BeaconClient->GetConnectionState() == EBeaconConnectionState::Pending);
-		UTEST_TRUE_EXPR(BeaconClient->GetNetConnection() != nullptr);
-
-		TSharedPtr<BeaconUnitTest::FNetworkStats> BeaconClientNetStats;
-		UTEST_TRUE_EXPR(BeaconUnitTest::ConfigureBeaconNetDriver(*Prerequisites, BeaconClient, &BeaconClientNetStats));
-		UTEST_TRUE_EXPR(BeaconClientNetStats.IsValid());
-
-		// Handshake testing.
-
-		UTEST_TRUE_EXPR(BeaconUnitTest::TickUntilControlMessageReceived(*Prerequisites, *BeaconHostNetStats));
-		UTEST_TRUE_EXPR(BeaconClient->GetConnectionState() == EBeaconConnectionState::Pending);
-		UTEST_TRUE_EXPR(BeaconClientNetStats->ReceivedControlMessages.IsEmpty());
-		UTEST_TRUE_EXPR(BeaconHostNetStats->ReceivedControlMessages == TArray<uint8>({ NMT_Hello }));
-		UTEST_TRUE_EXPR(BeaconUnitTest::GetBeaconClientForUser(BeaconHostObject, UserId) == nullptr);
-		BeaconHostNetStats->ReceivedControlMessages.Empty();
-
-		UTEST_TRUE_EXPR(BeaconUnitTest::TickOnce(*Prerequisites));
-		UTEST_TRUE_EXPR(BeaconClient->GetConnectionState() == EBeaconConnectionState::Pending);
-		UTEST_TRUE_EXPR(BeaconClientNetStats->ReceivedControlMessages == TArray<uint8>({ NMT_Challenge }));
-		UTEST_TRUE_EXPR(BeaconHostNetStats->ReceivedControlMessages.IsEmpty());
-		BeaconClientNetStats->ReceivedControlMessages.Empty();
-
-		UTEST_TRUE_EXPR(BeaconUnitTest::TickOnce(*Prerequisites));
-		UTEST_TRUE_EXPR(BeaconClient->GetConnectionState() == EBeaconConnectionState::Pending);
-		UTEST_TRUE_EXPR(BeaconClientNetStats->ReceivedControlMessages.IsEmpty());
-		UTEST_TRUE_EXPR(BeaconHostNetStats->ReceivedControlMessages == TArray<uint8>({ NMT_Login }));
-		UTEST_TRUE_EXPR(BeaconUnitTest::GetBeaconClientForUser(BeaconHostObject, UserId) == nullptr);
-		BeaconHostNetStats->ReceivedControlMessages.Empty();
-
-		UTEST_TRUE_EXPR(BeaconUnitTest::TickOnce(*Prerequisites));
-		UTEST_TRUE_EXPR(BeaconClient->GetConnectionState() == EBeaconConnectionState::Pending);
-		UTEST_TRUE_EXPR(BeaconClientNetStats->ReceivedControlMessages == TArray<uint8>({ NMT_BeaconWelcome }));
-		UTEST_TRUE_EXPR(BeaconHostNetStats->ReceivedControlMessages.IsEmpty());
-		BeaconClientNetStats->ReceivedControlMessages.Empty();
-
-		UTEST_TRUE_EXPR(BeaconUnitTest::TickOnce(*Prerequisites));
-		UTEST_TRUE_EXPR(BeaconClient->GetConnectionState() == EBeaconConnectionState::Pending);
-		UTEST_TRUE_EXPR(BeaconClientNetStats->ReceivedControlMessages.IsEmpty());
-		UTEST_TRUE_EXPR(BeaconHostNetStats->ReceivedControlMessages == TArray<uint8>({ NMT_Netspeed, NMT_BeaconJoin }));
-		UTEST_TRUE_EXPR(BeaconUnitTest::GetBeaconClientForUser(BeaconHostObject, UserId) == nullptr);
-		BeaconHostNetStats->ReceivedControlMessages.Empty();
-
-		UTEST_TRUE_EXPR(BeaconUnitTest::TickOnce(*Prerequisites));
-		UTEST_TRUE_EXPR(BeaconClient->GetConnectionState() == EBeaconConnectionState::Pending);
-		UTEST_TRUE_EXPR(BeaconClientNetStats->ReceivedControlMessages == TArray<uint8>({ NMT_BeaconAssignGUID }));
-		UTEST_TRUE_EXPR(BeaconHostNetStats->ReceivedControlMessages.IsEmpty());
-		BeaconClientNetStats->ReceivedControlMessages.Empty();
-
-		UTEST_TRUE_EXPR(BeaconUnitTest::TickOnce(*Prerequisites));
-		UTEST_TRUE_EXPR(BeaconClient->GetConnectionState() == EBeaconConnectionState::Pending);
-		UTEST_TRUE_EXPR(BeaconClientNetStats->ReceivedControlMessages.IsEmpty());
-		UTEST_TRUE_EXPR(BeaconHostNetStats->ReceivedControlMessages == TArray<uint8>({ NMT_BeaconNetGUIDAck }));
-		HostUserBeaconClient = BeaconUnitTest::GetBeaconClientForUser(BeaconHostObject, UserId);
-		UTEST_TRUE_EXPR(HostUserBeaconClient != nullptr);
-		UTEST_TRUE_EXPR(HostUserBeaconClient->GetUniqueId() == *UserId);
-		BeaconHostNetStats->ReceivedControlMessages.Empty();
-
-		// Receive OnConnected RPC on the client.
-		UTEST_TRUE_EXPR(BeaconUnitTest::TickOnce(*Prerequisites));
-		UTEST_TRUE_EXPR(BeaconClient->GetConnectionState() == EBeaconConnectionState::Open);
-		UTEST_TRUE_EXPR(BeaconClientNetStats->ReceivedControlMessages.IsEmpty());
-		UTEST_TRUE_EXPR(BeaconHostNetStats->ReceivedControlMessages.IsEmpty());
-
-		// Handshake complete, channel is open.
-
-		// Disconnect the client.
-		BeaconClient->DestroyBeacon();
-		BeaconClient = nullptr;
-
-		// Check that client actor on the host cleaned up.
-		UTEST_TRUE_EXPR(BeaconUnitTest::TickUntil(*Prerequisites, [&Prerequisites](){ return Prerequisites->GetStats().HostObject.NotifyClientDisconnected.InvokeCount == 1; }));
-		HostUserBeaconClient = BeaconUnitTest::GetBeaconClientForUser(BeaconHostObject, UserId);
-		UTEST_TRUE_EXPR(HostUserBeaconClient == nullptr);
-	}
-
-	// Method 1.
-	// Host handshake calls StartVerifyAuthentication and expects OnAuthenticationVerificationComplete to be called to complete authentication.
-	// Test OnAuthenticationVerificationComplete fired with a single frame delay is working.
-	{
-		BeaconHostNetStats->ReceivedControlMessages.Empty();
-		Prerequisites->GetConfig() = BaseConfig;
-		Prerequisites->GetStats() = BeaconUnitTest::FTestStats();
-
-		Prerequisites->GetConfig().Auth.Method1.bEnabled = true;
-		Prerequisites->GetConfig().Auth.Method1.bDelayDelegate = true;
-		Prerequisites->GetConfig().Auth.Method1.Result = FOnlineError::Success();
-
-		TObjectPtr<AOnlineBeaconClient> HostUserBeaconClient;
-
-		// Client setup.
-
-		TObjectPtr<AOnlineBeaconUnitTestClient> BeaconClient = Prerequisites->GetWorld()->SpawnActor<AOnlineBeaconUnitTestClient>(AOnlineBeaconUnitTestClient::StaticClass());
-		ON_SCOPE_EXIT{ if (BeaconClient) { BeaconClient->DestroyBeacon(); } };
-		UTEST_TRUE_EXPR(BeaconUnitTest::ConfigureBeacon(*Prerequisites, BeaconClient));
-
-		UTEST_TRUE_EXPR(BeaconClient->GetConnectionState() == EBeaconConnectionState::Invalid);
-		UTEST_TRUE_EXPR(BeaconUnitTest::InitClientForUser(*Prerequisites, BeaconClient, UserId));
-		UTEST_TRUE_EXPR(BeaconClient->GetConnectionState() == EBeaconConnectionState::Pending);
-		UTEST_TRUE_EXPR(BeaconClient->GetNetConnection() != nullptr);
-
-		TSharedPtr<BeaconUnitTest::FNetworkStats> BeaconClientNetStats;
-		UTEST_TRUE_EXPR(BeaconUnitTest::ConfigureBeaconNetDriver(*Prerequisites, BeaconClient, &BeaconClientNetStats));
-		UTEST_TRUE_EXPR(BeaconClientNetStats.IsValid());
-
-		// Handshake testing.
-
-		UTEST_TRUE_EXPR(BeaconUnitTest::TickUntilControlMessageReceived(*Prerequisites, *BeaconHostNetStats));
-		UTEST_TRUE_EXPR(BeaconClient->GetConnectionState() == EBeaconConnectionState::Pending);
-		UTEST_TRUE_EXPR(BeaconClientNetStats->ReceivedControlMessages.IsEmpty());
-		UTEST_TRUE_EXPR(BeaconHostNetStats->ReceivedControlMessages == TArray<uint8>({ NMT_Hello }));
-		UTEST_TRUE_EXPR(BeaconUnitTest::GetBeaconClientForUser(BeaconHostObject, UserId) == nullptr);
-		BeaconHostNetStats->ReceivedControlMessages.Empty();
-
-		UTEST_TRUE_EXPR(BeaconUnitTest::TickOnce(*Prerequisites));
-		UTEST_TRUE_EXPR(BeaconClient->GetConnectionState() == EBeaconConnectionState::Pending);
-		UTEST_TRUE_EXPR(BeaconClientNetStats->ReceivedControlMessages == TArray<uint8>({ NMT_Challenge }));
-		UTEST_TRUE_EXPR(BeaconHostNetStats->ReceivedControlMessages.IsEmpty());
-		BeaconClientNetStats->ReceivedControlMessages.Empty();
-
-		UTEST_TRUE_EXPR(BeaconUnitTest::TickOnce(*Prerequisites));
-		UTEST_TRUE_EXPR(BeaconClient->GetConnectionState() == EBeaconConnectionState::Pending);
-		UTEST_TRUE_EXPR(BeaconClientNetStats->ReceivedControlMessages.IsEmpty());
-		UTEST_TRUE_EXPR(BeaconHostNetStats->ReceivedControlMessages == TArray<uint8>({ NMT_Login }));
-		UTEST_TRUE_EXPR(BeaconUnitTest::GetBeaconClientForUser(BeaconHostObject, UserId) == nullptr);
-		BeaconHostNetStats->ReceivedControlMessages.Empty();
-
-		UTEST_TRUE_EXPR(BeaconUnitTest::TickOnce(*Prerequisites));
-		UTEST_TRUE_EXPR(BeaconClient->GetConnectionState() == EBeaconConnectionState::Pending);
-		UTEST_TRUE_EXPR(BeaconClientNetStats->ReceivedControlMessages.IsEmpty());
-		UTEST_TRUE_EXPR(BeaconHostNetStats->ReceivedControlMessages.IsEmpty());
-
-		UTEST_TRUE_EXPR(BeaconUnitTest::TickOnce(*Prerequisites));
-		UTEST_TRUE_EXPR(BeaconClient->GetConnectionState() == EBeaconConnectionState::Pending);
-		UTEST_TRUE_EXPR(BeaconClientNetStats->ReceivedControlMessages == TArray<uint8>({ NMT_BeaconWelcome }));
-		UTEST_TRUE_EXPR(BeaconHostNetStats->ReceivedControlMessages.IsEmpty());
-		BeaconClientNetStats->ReceivedControlMessages.Empty();
-
-		UTEST_TRUE_EXPR(BeaconUnitTest::TickOnce(*Prerequisites));
-		UTEST_TRUE_EXPR(BeaconClient->GetConnectionState() == EBeaconConnectionState::Pending);
-		UTEST_TRUE_EXPR(BeaconClientNetStats->ReceivedControlMessages.IsEmpty());
-		UTEST_TRUE_EXPR(BeaconHostNetStats->ReceivedControlMessages == TArray<uint8>({ NMT_Netspeed, NMT_BeaconJoin }));
-		UTEST_TRUE_EXPR(BeaconUnitTest::GetBeaconClientForUser(BeaconHostObject, UserId) == nullptr);
-		BeaconHostNetStats->ReceivedControlMessages.Empty();
-
-		UTEST_TRUE_EXPR(BeaconUnitTest::TickOnce(*Prerequisites));
-		UTEST_TRUE_EXPR(BeaconClient->GetConnectionState() == EBeaconConnectionState::Pending);
-		UTEST_TRUE_EXPR(BeaconClientNetStats->ReceivedControlMessages == TArray<uint8>({ NMT_BeaconAssignGUID }));
-		UTEST_TRUE_EXPR(BeaconHostNetStats->ReceivedControlMessages.IsEmpty());
-		BeaconClientNetStats->ReceivedControlMessages.Empty();
-
-		UTEST_TRUE_EXPR(BeaconUnitTest::TickOnce(*Prerequisites));
-		UTEST_TRUE_EXPR(BeaconClient->GetConnectionState() == EBeaconConnectionState::Pending);
-		UTEST_TRUE_EXPR(BeaconClientNetStats->ReceivedControlMessages.IsEmpty());
-		UTEST_TRUE_EXPR(BeaconHostNetStats->ReceivedControlMessages == TArray<uint8>({ NMT_BeaconNetGUIDAck }));
-		HostUserBeaconClient = BeaconUnitTest::GetBeaconClientForUser(BeaconHostObject, UserId);
-		UTEST_TRUE_EXPR(HostUserBeaconClient != nullptr);
-		UTEST_TRUE_EXPR(HostUserBeaconClient->GetUniqueId() == *UserId);
-		BeaconHostNetStats->ReceivedControlMessages.Empty();
-
-		// Receive OnConnected RPC on the client.
-		UTEST_TRUE_EXPR(BeaconUnitTest::TickOnce(*Prerequisites));
-		UTEST_TRUE_EXPR(BeaconClient->GetConnectionState() == EBeaconConnectionState::Open);
-		UTEST_TRUE_EXPR(BeaconClientNetStats->ReceivedControlMessages.IsEmpty());
-		UTEST_TRUE_EXPR(BeaconHostNetStats->ReceivedControlMessages.IsEmpty());
-
-		// Handshake complete, channel is open.
-
-		// Disconnect the client.
-		BeaconClient->DestroyBeacon();
-		BeaconClient = nullptr;
-
-		// Check that client actor on the host cleaned up.
-		UTEST_TRUE_EXPR(BeaconUnitTest::TickUntil(*Prerequisites, [&Prerequisites](){ return Prerequisites->GetStats().HostObject.NotifyClientDisconnected.InvokeCount == 1; }));
-		HostUserBeaconClient = BeaconUnitTest::GetBeaconClientForUser(BeaconHostObject, UserId);
-		UTEST_TRUE_EXPR(HostUserBeaconClient == nullptr);
-	}
-	
-	// Method 2.
-	// Host handshake calls StartVerifyAuthentication (callback signature) and expects callback delegate to be called to complete authentication.
-	// Test callback delegate fired with no delay is working.
-	{
-		BeaconHostNetStats->ReceivedControlMessages.Empty();
-		Prerequisites->GetConfig() = BaseConfig;
-		Prerequisites->GetStats() = BeaconUnitTest::FTestStats();
-
-		Prerequisites->GetConfig().Auth.Method2.bEnabled = true;
-		Prerequisites->GetConfig().Auth.Method2.Result = FOnlineError::Success();
-
-		TObjectPtr<AOnlineBeaconClient> HostUserBeaconClient;
-
-		// Client setup.
-
-		TObjectPtr<AOnlineBeaconUnitTestClient> BeaconClient = Prerequisites->GetWorld()->SpawnActor<AOnlineBeaconUnitTestClient>(AOnlineBeaconUnitTestClient::StaticClass());
-		ON_SCOPE_EXIT{ if (BeaconClient) { BeaconClient->DestroyBeacon(); } };
-		UTEST_TRUE_EXPR(BeaconUnitTest::ConfigureBeacon(*Prerequisites, BeaconClient));
-
-		UTEST_TRUE_EXPR(BeaconClient->GetConnectionState() == EBeaconConnectionState::Invalid);
-		UTEST_TRUE_EXPR(BeaconUnitTest::InitClientForUser(*Prerequisites, BeaconClient, UserId));
-		UTEST_TRUE_EXPR(BeaconClient->GetConnectionState() == EBeaconConnectionState::Pending);
-		UTEST_TRUE_EXPR(BeaconClient->GetNetConnection() != nullptr);
-
-		TSharedPtr<BeaconUnitTest::FNetworkStats> BeaconClientNetStats;
-		UTEST_TRUE_EXPR(BeaconUnitTest::ConfigureBeaconNetDriver(*Prerequisites, BeaconClient, &BeaconClientNetStats));
-		UTEST_TRUE_EXPR(BeaconClientNetStats.IsValid());
-
-		// Handshake testing.
-
-		UTEST_TRUE_EXPR(BeaconUnitTest::TickUntilControlMessageReceived(*Prerequisites, *BeaconHostNetStats));
-		UTEST_TRUE_EXPR(BeaconClient->GetConnectionState() == EBeaconConnectionState::Pending);
-		UTEST_TRUE_EXPR(BeaconClientNetStats->ReceivedControlMessages.IsEmpty());
-		UTEST_TRUE_EXPR(BeaconHostNetStats->ReceivedControlMessages == TArray<uint8>({ NMT_Hello }));
-		UTEST_TRUE_EXPR(BeaconUnitTest::GetBeaconClientForUser(BeaconHostObject, UserId) == nullptr);
-		BeaconHostNetStats->ReceivedControlMessages.Empty();
-
-		UTEST_TRUE_EXPR(BeaconUnitTest::TickOnce(*Prerequisites));
-		UTEST_TRUE_EXPR(BeaconClient->GetConnectionState() == EBeaconConnectionState::Pending);
-		UTEST_TRUE_EXPR(BeaconClientNetStats->ReceivedControlMessages == TArray<uint8>({ NMT_Challenge }));
-		UTEST_TRUE_EXPR(BeaconHostNetStats->ReceivedControlMessages.IsEmpty());
-		BeaconClientNetStats->ReceivedControlMessages.Empty();
-
-		UTEST_TRUE_EXPR(BeaconUnitTest::TickOnce(*Prerequisites));
-		UTEST_TRUE_EXPR(BeaconClient->GetConnectionState() == EBeaconConnectionState::Pending);
-		UTEST_TRUE_EXPR(BeaconClientNetStats->ReceivedControlMessages.IsEmpty());
-		UTEST_TRUE_EXPR(BeaconHostNetStats->ReceivedControlMessages == TArray<uint8>({ NMT_Login }));
-		UTEST_TRUE_EXPR(BeaconUnitTest::GetBeaconClientForUser(BeaconHostObject, UserId) == nullptr);
-		BeaconHostNetStats->ReceivedControlMessages.Empty();
-
-		UTEST_TRUE_EXPR(BeaconUnitTest::TickOnce(*Prerequisites));
-		UTEST_TRUE_EXPR(BeaconClient->GetConnectionState() == EBeaconConnectionState::Pending);
-		UTEST_TRUE_EXPR(BeaconClientNetStats->ReceivedControlMessages == TArray<uint8>({ NMT_BeaconWelcome }));
-		UTEST_TRUE_EXPR(BeaconHostNetStats->ReceivedControlMessages.IsEmpty());
-		BeaconClientNetStats->ReceivedControlMessages.Empty();
-
-		UTEST_TRUE_EXPR(BeaconUnitTest::TickOnce(*Prerequisites));
-		UTEST_TRUE_EXPR(BeaconClient->GetConnectionState() == EBeaconConnectionState::Pending);
-		UTEST_TRUE_EXPR(BeaconClientNetStats->ReceivedControlMessages.IsEmpty());
-		UTEST_TRUE_EXPR(BeaconHostNetStats->ReceivedControlMessages == TArray<uint8>({ NMT_Netspeed, NMT_BeaconJoin }));
-		UTEST_TRUE_EXPR(BeaconUnitTest::GetBeaconClientForUser(BeaconHostObject, UserId) == nullptr);
-		BeaconHostNetStats->ReceivedControlMessages.Empty();
-
-		UTEST_TRUE_EXPR(BeaconUnitTest::TickOnce(*Prerequisites));
-		UTEST_TRUE_EXPR(BeaconClient->GetConnectionState() == EBeaconConnectionState::Pending);
-		UTEST_TRUE_EXPR(BeaconClientNetStats->ReceivedControlMessages == TArray<uint8>({ NMT_BeaconAssignGUID }));
-		UTEST_TRUE_EXPR(BeaconHostNetStats->ReceivedControlMessages.IsEmpty());
-		BeaconClientNetStats->ReceivedControlMessages.Empty();
-
-		UTEST_TRUE_EXPR(BeaconUnitTest::TickOnce(*Prerequisites));
-		UTEST_TRUE_EXPR(BeaconClient->GetConnectionState() == EBeaconConnectionState::Pending);
-		UTEST_TRUE_EXPR(BeaconClientNetStats->ReceivedControlMessages.IsEmpty());
-		UTEST_TRUE_EXPR(BeaconHostNetStats->ReceivedControlMessages == TArray<uint8>({ NMT_BeaconNetGUIDAck }));
-		HostUserBeaconClient = BeaconUnitTest::GetBeaconClientForUser(BeaconHostObject, UserId);
-		UTEST_TRUE_EXPR(HostUserBeaconClient != nullptr);
-		UTEST_TRUE_EXPR(HostUserBeaconClient->GetUniqueId() == *UserId);
-		BeaconHostNetStats->ReceivedControlMessages.Empty();
-
-		// Receive OnConnected RPC on the client.
-		UTEST_TRUE_EXPR(BeaconUnitTest::TickOnce(*Prerequisites));
-		UTEST_TRUE_EXPR(BeaconClient->GetConnectionState() == EBeaconConnectionState::Open);
-		UTEST_TRUE_EXPR(BeaconClientNetStats->ReceivedControlMessages.IsEmpty());
-		UTEST_TRUE_EXPR(BeaconHostNetStats->ReceivedControlMessages.IsEmpty());
-
-		// Handshake complete, channel is open.
-
-		// Disconnect the client.
-		BeaconClient->DestroyBeacon();
-		BeaconClient = nullptr;
-
-		// Check that client actor on the host cleaned up.
-		UTEST_TRUE_EXPR(BeaconUnitTest::TickUntil(*Prerequisites, [&Prerequisites](){ return Prerequisites->GetStats().HostObject.NotifyClientDisconnected.InvokeCount == 1; }));
-		HostUserBeaconClient = BeaconUnitTest::GetBeaconClientForUser(BeaconHostObject, UserId);
-		UTEST_TRUE_EXPR(HostUserBeaconClient == nullptr);
-	}
-
-	// Method 2.
-	// Host handshake calls StartVerifyAuthentication (callback signature) and expects callback delegate to be called to complete authentication.
-	// Test callback delegate fired with a single frame delay is working.
-	{
-		BeaconHostNetStats->ReceivedControlMessages.Empty();
-		Prerequisites->GetConfig() = BaseConfig;
-		Prerequisites->GetStats() = BeaconUnitTest::FTestStats();
-
-		Prerequisites->GetConfig().Auth.Method2.bEnabled = true;
-		Prerequisites->GetConfig().Auth.Method2.bDelayDelegate = true;
-		Prerequisites->GetConfig().Auth.Method2.Result = FOnlineError::Success();
-
-		TObjectPtr<AOnlineBeaconClient> HostUserBeaconClient;
-
-		// Client setup.
-
-		TObjectPtr<AOnlineBeaconUnitTestClient> BeaconClient = Prerequisites->GetWorld()->SpawnActor<AOnlineBeaconUnitTestClient>(AOnlineBeaconUnitTestClient::StaticClass());
-		ON_SCOPE_EXIT{ if (BeaconClient) { BeaconClient->DestroyBeacon(); } };
-		UTEST_TRUE_EXPR(BeaconUnitTest::ConfigureBeacon(*Prerequisites, BeaconClient));
-
-		UTEST_TRUE_EXPR(BeaconClient->GetConnectionState() == EBeaconConnectionState::Invalid);
-		UTEST_TRUE_EXPR(BeaconUnitTest::InitClientForUser(*Prerequisites, BeaconClient, UserId));
-		UTEST_TRUE_EXPR(BeaconClient->GetConnectionState() == EBeaconConnectionState::Pending);
-		UTEST_TRUE_EXPR(BeaconClient->GetNetConnection() != nullptr);
-
-		TSharedPtr<BeaconUnitTest::FNetworkStats> BeaconClientNetStats;
-		UTEST_TRUE_EXPR(BeaconUnitTest::ConfigureBeaconNetDriver(*Prerequisites, BeaconClient, &BeaconClientNetStats));
-		UTEST_TRUE_EXPR(BeaconClientNetStats.IsValid());
-
-		// Handshake testing.
-
-		UTEST_TRUE_EXPR(BeaconUnitTest::TickUntilControlMessageReceived(*Prerequisites, *BeaconHostNetStats));
-		UTEST_TRUE_EXPR(BeaconClient->GetConnectionState() == EBeaconConnectionState::Pending);
-		UTEST_TRUE_EXPR(BeaconClientNetStats->ReceivedControlMessages.IsEmpty());
-		UTEST_TRUE_EXPR(BeaconHostNetStats->ReceivedControlMessages == TArray<uint8>({ NMT_Hello }));
-		UTEST_TRUE_EXPR(BeaconUnitTest::GetBeaconClientForUser(BeaconHostObject, UserId) == nullptr);
-		BeaconHostNetStats->ReceivedControlMessages.Empty();
-
-		UTEST_TRUE_EXPR(BeaconUnitTest::TickOnce(*Prerequisites));
-		UTEST_TRUE_EXPR(BeaconClient->GetConnectionState() == EBeaconConnectionState::Pending);
-		UTEST_TRUE_EXPR(BeaconClientNetStats->ReceivedControlMessages == TArray<uint8>({ NMT_Challenge }));
-		UTEST_TRUE_EXPR(BeaconHostNetStats->ReceivedControlMessages.IsEmpty());
-		BeaconClientNetStats->ReceivedControlMessages.Empty();
-
-		UTEST_TRUE_EXPR(BeaconUnitTest::TickOnce(*Prerequisites));
-		UTEST_TRUE_EXPR(BeaconClient->GetConnectionState() == EBeaconConnectionState::Pending);
-		UTEST_TRUE_EXPR(BeaconClientNetStats->ReceivedControlMessages.IsEmpty());
-		UTEST_TRUE_EXPR(BeaconHostNetStats->ReceivedControlMessages == TArray<uint8>({ NMT_Login }));
-		UTEST_TRUE_EXPR(BeaconUnitTest::GetBeaconClientForUser(BeaconHostObject, UserId) == nullptr);
-		BeaconHostNetStats->ReceivedControlMessages.Empty();
-
-		UTEST_TRUE_EXPR(BeaconUnitTest::TickOnce(*Prerequisites));
-		UTEST_TRUE_EXPR(BeaconClient->GetConnectionState() == EBeaconConnectionState::Pending);
-		UTEST_TRUE_EXPR(BeaconClientNetStats->ReceivedControlMessages.IsEmpty());
-		UTEST_TRUE_EXPR(BeaconHostNetStats->ReceivedControlMessages.IsEmpty());
-
-		UTEST_TRUE_EXPR(BeaconUnitTest::TickOnce(*Prerequisites));
-		UTEST_TRUE_EXPR(BeaconClient->GetConnectionState() == EBeaconConnectionState::Pending);
-		UTEST_TRUE_EXPR(BeaconClientNetStats->ReceivedControlMessages == TArray<uint8>({ NMT_BeaconWelcome }));
-		UTEST_TRUE_EXPR(BeaconHostNetStats->ReceivedControlMessages.IsEmpty());
-		BeaconClientNetStats->ReceivedControlMessages.Empty();
-
-		UTEST_TRUE_EXPR(BeaconUnitTest::TickOnce(*Prerequisites));
-		UTEST_TRUE_EXPR(BeaconClient->GetConnectionState() == EBeaconConnectionState::Pending);
-		UTEST_TRUE_EXPR(BeaconClientNetStats->ReceivedControlMessages.IsEmpty());
-		UTEST_TRUE_EXPR(BeaconHostNetStats->ReceivedControlMessages == TArray<uint8>({ NMT_Netspeed, NMT_BeaconJoin }));
-		UTEST_TRUE_EXPR(BeaconUnitTest::GetBeaconClientForUser(BeaconHostObject, UserId) == nullptr);
-		BeaconHostNetStats->ReceivedControlMessages.Empty();
-
-		UTEST_TRUE_EXPR(BeaconUnitTest::TickOnce(*Prerequisites));
-		UTEST_TRUE_EXPR(BeaconClient->GetConnectionState() == EBeaconConnectionState::Pending);
-		UTEST_TRUE_EXPR(BeaconClientNetStats->ReceivedControlMessages == TArray<uint8>({ NMT_BeaconAssignGUID }));
-		UTEST_TRUE_EXPR(BeaconHostNetStats->ReceivedControlMessages.IsEmpty());
-		BeaconClientNetStats->ReceivedControlMessages.Empty();
-
-		UTEST_TRUE_EXPR(BeaconUnitTest::TickOnce(*Prerequisites));
-		UTEST_TRUE_EXPR(BeaconClient->GetConnectionState() == EBeaconConnectionState::Pending);
-		UTEST_TRUE_EXPR(BeaconClientNetStats->ReceivedControlMessages.IsEmpty());
-		UTEST_TRUE_EXPR(BeaconHostNetStats->ReceivedControlMessages == TArray<uint8>({ NMT_BeaconNetGUIDAck }));
-		HostUserBeaconClient = BeaconUnitTest::GetBeaconClientForUser(BeaconHostObject, UserId);
-		UTEST_TRUE_EXPR(HostUserBeaconClient != nullptr);
-		UTEST_TRUE_EXPR(HostUserBeaconClient->GetUniqueId() == *UserId);
-		BeaconHostNetStats->ReceivedControlMessages.Empty();
-
-		// Receive OnConnected RPC on the client.
-		UTEST_TRUE_EXPR(BeaconUnitTest::TickOnce(*Prerequisites));
-		UTEST_TRUE_EXPR(BeaconClient->GetConnectionState() == EBeaconConnectionState::Open);
-		UTEST_TRUE_EXPR(BeaconClientNetStats->ReceivedControlMessages.IsEmpty());
-		UTEST_TRUE_EXPR(BeaconHostNetStats->ReceivedControlMessages.IsEmpty());
-
-		// Handshake complete, channel is open.
-
-		// Disconnect the client.
-		BeaconClient->DestroyBeacon();
-		BeaconClient = nullptr;
-
-		// Check that client actor on the host cleaned up.
-		UTEST_TRUE_EXPR(BeaconUnitTest::TickUntil(*Prerequisites, [&Prerequisites](){ return Prerequisites->GetStats().HostObject.NotifyClientDisconnected.InvokeCount == 1; }));
-		HostUserBeaconClient = BeaconUnitTest::GetBeaconClientForUser(BeaconHostObject, UserId);
-		UTEST_TRUE_EXPR(HostUserBeaconClient == nullptr);
-	}
-
-	// Method 3.
 	// Host handshake calls StartVerifyAuthentication (callback signature with LoginOptions) and expects callback delegate to be called to complete authentication.
 	// Test callback delegate fired with no delay is working.
 	{
@@ -1963,8 +1566,7 @@ bool FOnlineBeaconTestAuthenticatedHandshakeSuccess::RunTest(const FString& Para
 		Prerequisites->GetConfig() = BaseConfig;
 		Prerequisites->GetStats() = BeaconUnitTest::FTestStats();
 
-		Prerequisites->GetConfig().Auth.Method3.bEnabled = true;
-		Prerequisites->GetConfig().Auth.Method3.Result = FOnlineError::Success();
+		Prerequisites->GetConfig().Auth.Result = FOnlineError::Success();
 
 		TObjectPtr<AOnlineBeaconClient> HostUserBeaconClient;
 
@@ -2051,7 +1653,6 @@ bool FOnlineBeaconTestAuthenticatedHandshakeSuccess::RunTest(const FString& Para
 		UTEST_TRUE_EXPR(HostUserBeaconClient == nullptr);
 	}
 
-	// Method 3.
 	// Host handshake calls StartVerifyAuthentication (callback signature with LoginOptions) and expects callback delegate to be called to complete authentication.
 	// Test callback delegate fired with a single frame delay is working.
 	{
@@ -2059,9 +1660,9 @@ bool FOnlineBeaconTestAuthenticatedHandshakeSuccess::RunTest(const FString& Para
 		Prerequisites->GetConfig() = BaseConfig;
 		Prerequisites->GetStats() = BeaconUnitTest::FTestStats();
 
-		Prerequisites->GetConfig().Auth.Method3.bEnabled = true;
-		Prerequisites->GetConfig().Auth.Method3.bDelayDelegate = true;
-		Prerequisites->GetConfig().Auth.Method3.Result = FOnlineError::Success();
+		Prerequisites->GetConfig().Auth.bEnabled = true;
+		Prerequisites->GetConfig().Auth.bDelayDelegate = true;
+		Prerequisites->GetConfig().Auth.Result = FOnlineError::Success();
 
 		TObjectPtr<AOnlineBeaconClient> HostUserBeaconClient;
 
@@ -2153,15 +1754,14 @@ bool FOnlineBeaconTestAuthenticatedHandshakeSuccess::RunTest(const FString& Para
 		UTEST_TRUE_EXPR(HostUserBeaconClient == nullptr);
 	}
 
-	// Auth verify.
-	// Use Method 3 auth and verify that a user supplied VerifyJoinForBeaconType method allows the connection to proceed.
+	// Auth and verify that a user supplied VerifyJoinForBeaconType method allows the connection to proceed.
 	{
 		BeaconHostNetStats->ReceivedControlMessages.Empty();
 		Prerequisites->GetConfig() = BaseConfig;
 		Prerequisites->GetStats() = BeaconUnitTest::FTestStats();
 
-		Prerequisites->GetConfig().Auth.Method3.bEnabled = true;
-		Prerequisites->GetConfig().Auth.Method3.Result = FOnlineError::Success();
+		Prerequisites->GetConfig().Auth.bEnabled = true;
+		Prerequisites->GetConfig().Auth.Result = FOnlineError::Success();
 
 		Prerequisites->GetConfig().Auth.Verify.bEnabled = true;
 		Prerequisites->GetConfig().Auth.Verify.bResult = true;
@@ -2285,243 +1885,6 @@ bool FOnlineBeaconTestAuthenticatedHandshakeFailure::RunTest(const FString& Para
 	UTEST_TRUE_EXPR(BeaconUnitTest::ConfigureBeaconNetDriver(*Prerequisites, BeaconHost, &BeaconHostNetStats));
 	UTEST_TRUE_EXPR(BeaconHostNetStats.IsValid());
 
-	// Method 1.
-	// Host handshake calls StartVerifyAuthentication and expects OnAuthenticationVerificationComplete to be called to complete authentication.
-	// Test OnAuthenticationVerificationComplete fired with no delay is working.
-	{
-		BeaconHostNetStats->ReceivedControlMessages.Empty();
-		Prerequisites->GetConfig() = BaseConfig;
-		Prerequisites->GetStats() = BeaconUnitTest::FTestStats();
-
-		Prerequisites->GetConfig().Auth.Method1.bEnabled = true;
-		Prerequisites->GetConfig().Auth.Method1.Result = OnlineIdentity::Errors::InvalidCreds();
-
-		// Client setup.
-
-		TObjectPtr<AOnlineBeaconUnitTestClient> BeaconClient = Prerequisites->GetWorld()->SpawnActor<AOnlineBeaconUnitTestClient>(AOnlineBeaconUnitTestClient::StaticClass());
-		ON_SCOPE_EXIT{ if (BeaconClient) { BeaconClient->DestroyBeacon(); } };
-		UTEST_TRUE_EXPR(BeaconUnitTest::ConfigureBeacon(*Prerequisites, BeaconClient));
-
-		UTEST_TRUE_EXPR(BeaconClient->GetConnectionState() == EBeaconConnectionState::Invalid);
-		UTEST_TRUE_EXPR(BeaconUnitTest::InitClientForUser(*Prerequisites, BeaconClient, UserId));
-		UTEST_TRUE_EXPR(BeaconClient->GetConnectionState() == EBeaconConnectionState::Pending);
-		UTEST_TRUE_EXPR(BeaconClient->GetNetConnection() != nullptr);
-
-		TSharedPtr<BeaconUnitTest::FNetworkStats> BeaconClientNetStats;
-		UTEST_TRUE_EXPR(BeaconUnitTest::ConfigureBeaconNetDriver(*Prerequisites, BeaconClient, &BeaconClientNetStats));
-		UTEST_TRUE_EXPR(BeaconClientNetStats.IsValid());
-
-		// Handshake testing.
-
-		UTEST_TRUE_EXPR(BeaconUnitTest::TickUntilControlMessageReceived(*Prerequisites, *BeaconHostNetStats));
-		UTEST_TRUE_EXPR(BeaconClient->GetConnectionState() == EBeaconConnectionState::Pending);
-		UTEST_TRUE_EXPR(BeaconClientNetStats->ReceivedControlMessages.IsEmpty());
-		UTEST_TRUE_EXPR(BeaconHostNetStats->ReceivedControlMessages == TArray<uint8>({ NMT_Hello }));
-		UTEST_TRUE_EXPR(BeaconUnitTest::GetBeaconClientForUser(BeaconHostObject, UserId) == nullptr);
-		BeaconHostNetStats->ReceivedControlMessages.Empty();
-
-		UTEST_TRUE_EXPR(BeaconUnitTest::TickOnce(*Prerequisites));
-		UTEST_TRUE_EXPR(BeaconClient->GetConnectionState() == EBeaconConnectionState::Pending);
-		UTEST_TRUE_EXPR(BeaconClientNetStats->ReceivedControlMessages == TArray<uint8>({ NMT_Challenge }));
-		UTEST_TRUE_EXPR(BeaconHostNetStats->ReceivedControlMessages.IsEmpty());
-		BeaconClientNetStats->ReceivedControlMessages.Empty();
-
-		UTEST_TRUE_EXPR(BeaconUnitTest::TickOnce(*Prerequisites));
-		UTEST_TRUE_EXPR(BeaconClient->GetConnectionState() == EBeaconConnectionState::Pending);
-		UTEST_TRUE_EXPR(BeaconClientNetStats->ReceivedControlMessages.IsEmpty());
-		UTEST_TRUE_EXPR(BeaconHostNetStats->ReceivedControlMessages == TArray<uint8>({ NMT_Login }));
-		UTEST_TRUE_EXPR(BeaconUnitTest::GetBeaconClientForUser(BeaconHostObject, UserId) == nullptr);
-		BeaconHostNetStats->ReceivedControlMessages.Empty();
-
-		UTEST_TRUE_EXPR(BeaconUnitTest::TickOnce(*Prerequisites));
-		UTEST_TRUE_EXPR(BeaconClient->GetConnectionState() == EBeaconConnectionState::Invalid);
-		UTEST_TRUE_EXPR(BeaconClientNetStats->ReceivedControlMessages == TArray<uint8>({ NMT_Failure }));
-		UTEST_TRUE_EXPR(BeaconHostNetStats->ReceivedControlMessages.IsEmpty());
-		UTEST_TRUE_EXPR(Prerequisites->GetStats().Client.OnFailure.InvokeCount == 1);
-		BeaconClientNetStats->ReceivedControlMessages.Empty();
-	}
-
-	// Method 1.
-	// Host handshake calls StartVerifyAuthentication and expects OnAuthenticationVerificationComplete to be called to complete authentication.
-	// Test OnAuthenticationVerificationComplete fired with a single frame delay is working.
-	{
-		BeaconHostNetStats->ReceivedControlMessages.Empty();
-		Prerequisites->GetConfig() = BaseConfig;
-		Prerequisites->GetStats() = BeaconUnitTest::FTestStats();
-
-		Prerequisites->GetConfig().Auth.Method1.bEnabled = true;
-		Prerequisites->GetConfig().Auth.Method1.bDelayDelegate = true;
-		Prerequisites->GetConfig().Auth.Method1.Result = OnlineIdentity::Errors::InvalidCreds();
-
-		// Client setup.
-
-		TObjectPtr<AOnlineBeaconUnitTestClient> BeaconClient = Prerequisites->GetWorld()->SpawnActor<AOnlineBeaconUnitTestClient>(AOnlineBeaconUnitTestClient::StaticClass());
-		ON_SCOPE_EXIT{ if (BeaconClient) { BeaconClient->DestroyBeacon(); } };
-		UTEST_TRUE_EXPR(BeaconUnitTest::ConfigureBeacon(*Prerequisites, BeaconClient));
-
-		UTEST_TRUE_EXPR(BeaconClient->GetConnectionState() == EBeaconConnectionState::Invalid);
-		UTEST_TRUE_EXPR(BeaconUnitTest::InitClientForUser(*Prerequisites, BeaconClient, UserId));
-		UTEST_TRUE_EXPR(BeaconClient->GetConnectionState() == EBeaconConnectionState::Pending);
-		UTEST_TRUE_EXPR(BeaconClient->GetNetConnection() != nullptr);
-
-		TSharedPtr<BeaconUnitTest::FNetworkStats> BeaconClientNetStats;
-		UTEST_TRUE_EXPR(BeaconUnitTest::ConfigureBeaconNetDriver(*Prerequisites, BeaconClient, &BeaconClientNetStats));
-		UTEST_TRUE_EXPR(BeaconClientNetStats.IsValid());
-
-		// Handshake testing.
-
-		UTEST_TRUE_EXPR(BeaconUnitTest::TickUntilControlMessageReceived(*Prerequisites, *BeaconHostNetStats));
-		UTEST_TRUE_EXPR(BeaconClient->GetConnectionState() == EBeaconConnectionState::Pending);
-		UTEST_TRUE_EXPR(BeaconClientNetStats->ReceivedControlMessages.IsEmpty());
-		UTEST_TRUE_EXPR(BeaconHostNetStats->ReceivedControlMessages == TArray<uint8>({ NMT_Hello }));
-		UTEST_TRUE_EXPR(BeaconUnitTest::GetBeaconClientForUser(BeaconHostObject, UserId) == nullptr);
-		BeaconHostNetStats->ReceivedControlMessages.Empty();
-
-		UTEST_TRUE_EXPR(BeaconUnitTest::TickOnce(*Prerequisites));
-		UTEST_TRUE_EXPR(BeaconClient->GetConnectionState() == EBeaconConnectionState::Pending);
-		UTEST_TRUE_EXPR(BeaconClientNetStats->ReceivedControlMessages == TArray<uint8>({ NMT_Challenge }));
-		UTEST_TRUE_EXPR(BeaconHostNetStats->ReceivedControlMessages.IsEmpty());
-		BeaconClientNetStats->ReceivedControlMessages.Empty();
-
-		UTEST_TRUE_EXPR(BeaconUnitTest::TickOnce(*Prerequisites));
-		UTEST_TRUE_EXPR(BeaconClient->GetConnectionState() == EBeaconConnectionState::Pending);
-		UTEST_TRUE_EXPR(BeaconClientNetStats->ReceivedControlMessages.IsEmpty());
-		UTEST_TRUE_EXPR(BeaconHostNetStats->ReceivedControlMessages == TArray<uint8>({ NMT_Login }));
-		UTEST_TRUE_EXPR(BeaconUnitTest::GetBeaconClientForUser(BeaconHostObject, UserId) == nullptr);
-		BeaconHostNetStats->ReceivedControlMessages.Empty();
-
-		UTEST_TRUE_EXPR(BeaconUnitTest::TickOnce(*Prerequisites));
-		UTEST_TRUE_EXPR(BeaconClient->GetConnectionState() == EBeaconConnectionState::Pending);
-		UTEST_TRUE_EXPR(BeaconClientNetStats->ReceivedControlMessages.IsEmpty());
-		UTEST_TRUE_EXPR(BeaconHostNetStats->ReceivedControlMessages.IsEmpty());
-
-		UTEST_TRUE_EXPR(BeaconUnitTest::TickOnce(*Prerequisites));
-		UTEST_TRUE_EXPR(BeaconClient->GetConnectionState() == EBeaconConnectionState::Invalid);
-		UTEST_TRUE_EXPR(BeaconClientNetStats->ReceivedControlMessages == TArray<uint8>({ NMT_Failure }));
-		UTEST_TRUE_EXPR(BeaconHostNetStats->ReceivedControlMessages.IsEmpty());
-		UTEST_TRUE_EXPR(Prerequisites->GetStats().Client.OnFailure.InvokeCount == 1);
-		BeaconClientNetStats->ReceivedControlMessages.Empty();
-	}
-	
-	// Method 2.
-	// Host handshake calls StartVerifyAuthentication (callback signature) and expects callback delegate to be called to complete authentication.
-	// Test callback delegate fired with no delay is working.
-	{
-		BeaconHostNetStats->ReceivedControlMessages.Empty();
-		Prerequisites->GetConfig() = BaseConfig;
-		Prerequisites->GetStats() = BeaconUnitTest::FTestStats();
-
-		Prerequisites->GetConfig().Auth.Method2.bEnabled = true;
-		Prerequisites->GetConfig().Auth.Method2.Result = OnlineIdentity::Errors::InvalidCreds();
-
-		// Client setup.
-
-		TObjectPtr<AOnlineBeaconUnitTestClient> BeaconClient = Prerequisites->GetWorld()->SpawnActor<AOnlineBeaconUnitTestClient>(AOnlineBeaconUnitTestClient::StaticClass());
-		ON_SCOPE_EXIT{ if (BeaconClient) { BeaconClient->DestroyBeacon(); } };
-		UTEST_TRUE_EXPR(BeaconUnitTest::ConfigureBeacon(*Prerequisites, BeaconClient));
-
-		UTEST_TRUE_EXPR(BeaconClient->GetConnectionState() == EBeaconConnectionState::Invalid);
-		UTEST_TRUE_EXPR(BeaconUnitTest::InitClientForUser(*Prerequisites, BeaconClient, UserId));
-		UTEST_TRUE_EXPR(BeaconClient->GetConnectionState() == EBeaconConnectionState::Pending);
-		UTEST_TRUE_EXPR(BeaconClient->GetNetConnection() != nullptr);
-
-		TSharedPtr<BeaconUnitTest::FNetworkStats> BeaconClientNetStats;
-		UTEST_TRUE_EXPR(BeaconUnitTest::ConfigureBeaconNetDriver(*Prerequisites, BeaconClient, &BeaconClientNetStats));
-		UTEST_TRUE_EXPR(BeaconClientNetStats.IsValid());
-
-		// Handshake testing.
-
-		UTEST_TRUE_EXPR(BeaconUnitTest::TickUntilControlMessageReceived(*Prerequisites, *BeaconHostNetStats));
-		UTEST_TRUE_EXPR(BeaconClient->GetConnectionState() == EBeaconConnectionState::Pending);
-		UTEST_TRUE_EXPR(BeaconClientNetStats->ReceivedControlMessages.IsEmpty());
-		UTEST_TRUE_EXPR(BeaconHostNetStats->ReceivedControlMessages == TArray<uint8>({ NMT_Hello }));
-		UTEST_TRUE_EXPR(BeaconUnitTest::GetBeaconClientForUser(BeaconHostObject, UserId) == nullptr);
-		BeaconHostNetStats->ReceivedControlMessages.Empty();
-
-		UTEST_TRUE_EXPR(BeaconUnitTest::TickOnce(*Prerequisites));
-		UTEST_TRUE_EXPR(BeaconClient->GetConnectionState() == EBeaconConnectionState::Pending);
-		UTEST_TRUE_EXPR(BeaconClientNetStats->ReceivedControlMessages == TArray<uint8>({ NMT_Challenge }));
-		UTEST_TRUE_EXPR(BeaconHostNetStats->ReceivedControlMessages.IsEmpty());
-		BeaconClientNetStats->ReceivedControlMessages.Empty();
-
-		UTEST_TRUE_EXPR(BeaconUnitTest::TickOnce(*Prerequisites));
-		UTEST_TRUE_EXPR(BeaconClient->GetConnectionState() == EBeaconConnectionState::Pending);
-		UTEST_TRUE_EXPR(BeaconClientNetStats->ReceivedControlMessages.IsEmpty());
-		UTEST_TRUE_EXPR(BeaconHostNetStats->ReceivedControlMessages == TArray<uint8>({ NMT_Login }));
-		UTEST_TRUE_EXPR(BeaconUnitTest::GetBeaconClientForUser(BeaconHostObject, UserId) == nullptr);
-		BeaconHostNetStats->ReceivedControlMessages.Empty();
-
-		UTEST_TRUE_EXPR(BeaconUnitTest::TickOnce(*Prerequisites));
-		UTEST_TRUE_EXPR(BeaconClient->GetConnectionState() == EBeaconConnectionState::Invalid);
-		UTEST_TRUE_EXPR(BeaconClientNetStats->ReceivedControlMessages == TArray<uint8>({ NMT_Failure }));
-		UTEST_TRUE_EXPR(BeaconHostNetStats->ReceivedControlMessages.IsEmpty());
-		UTEST_TRUE_EXPR(Prerequisites->GetStats().Client.OnFailure.InvokeCount == 1);
-		BeaconClientNetStats->ReceivedControlMessages.Empty();
-	}
-
-	// Method 2.
-	// Host handshake calls StartVerifyAuthentication (callback signature) and expects callback delegate to be called to complete authentication.
-	// Test callback delegate fired with a single frame delay is working.
-	{
-		BeaconHostNetStats->ReceivedControlMessages.Empty();
-		Prerequisites->GetConfig() = BaseConfig;
-		Prerequisites->GetStats() = BeaconUnitTest::FTestStats();
-
-		Prerequisites->GetConfig().Auth.Method2.bEnabled = true;
-		Prerequisites->GetConfig().Auth.Method2.bDelayDelegate = true;
-		Prerequisites->GetConfig().Auth.Method2.Result = OnlineIdentity::Errors::InvalidCreds();
-
-		// Client setup.
-
-		TObjectPtr<AOnlineBeaconUnitTestClient> BeaconClient = Prerequisites->GetWorld()->SpawnActor<AOnlineBeaconUnitTestClient>(AOnlineBeaconUnitTestClient::StaticClass());
-		ON_SCOPE_EXIT{ if (BeaconClient) { BeaconClient->DestroyBeacon(); } };
-		UTEST_TRUE_EXPR(BeaconUnitTest::ConfigureBeacon(*Prerequisites, BeaconClient));
-
-		UTEST_TRUE_EXPR(BeaconClient->GetConnectionState() == EBeaconConnectionState::Invalid);
-		UTEST_TRUE_EXPR(BeaconUnitTest::InitClientForUser(*Prerequisites, BeaconClient, UserId));
-		UTEST_TRUE_EXPR(BeaconClient->GetConnectionState() == EBeaconConnectionState::Pending);
-		UTEST_TRUE_EXPR(BeaconClient->GetNetConnection() != nullptr);
-
-		TSharedPtr<BeaconUnitTest::FNetworkStats> BeaconClientNetStats;
-		UTEST_TRUE_EXPR(BeaconUnitTest::ConfigureBeaconNetDriver(*Prerequisites, BeaconClient, &BeaconClientNetStats));
-		UTEST_TRUE_EXPR(BeaconClientNetStats.IsValid());
-
-		// Handshake testing.
-
-		UTEST_TRUE_EXPR(BeaconUnitTest::TickUntilControlMessageReceived(*Prerequisites, *BeaconHostNetStats));
-		UTEST_TRUE_EXPR(BeaconClient->GetConnectionState() == EBeaconConnectionState::Pending);
-		UTEST_TRUE_EXPR(BeaconClientNetStats->ReceivedControlMessages.IsEmpty());
-		UTEST_TRUE_EXPR(BeaconHostNetStats->ReceivedControlMessages == TArray<uint8>({ NMT_Hello }));
-		UTEST_TRUE_EXPR(BeaconUnitTest::GetBeaconClientForUser(BeaconHostObject, UserId) == nullptr);
-		BeaconHostNetStats->ReceivedControlMessages.Empty();
-
-		UTEST_TRUE_EXPR(BeaconUnitTest::TickOnce(*Prerequisites));
-		UTEST_TRUE_EXPR(BeaconClient->GetConnectionState() == EBeaconConnectionState::Pending);
-		UTEST_TRUE_EXPR(BeaconClientNetStats->ReceivedControlMessages == TArray<uint8>({ NMT_Challenge }));
-		UTEST_TRUE_EXPR(BeaconHostNetStats->ReceivedControlMessages.IsEmpty());
-		BeaconClientNetStats->ReceivedControlMessages.Empty();
-
-		UTEST_TRUE_EXPR(BeaconUnitTest::TickOnce(*Prerequisites));
-		UTEST_TRUE_EXPR(BeaconClient->GetConnectionState() == EBeaconConnectionState::Pending);
-		UTEST_TRUE_EXPR(BeaconClientNetStats->ReceivedControlMessages.IsEmpty());
-		UTEST_TRUE_EXPR(BeaconHostNetStats->ReceivedControlMessages == TArray<uint8>({ NMT_Login }));
-		UTEST_TRUE_EXPR(BeaconUnitTest::GetBeaconClientForUser(BeaconHostObject, UserId) == nullptr);
-		BeaconHostNetStats->ReceivedControlMessages.Empty();
-
-		UTEST_TRUE_EXPR(BeaconUnitTest::TickOnce(*Prerequisites));
-		UTEST_TRUE_EXPR(BeaconClient->GetConnectionState() == EBeaconConnectionState::Pending);
-		UTEST_TRUE_EXPR(BeaconClientNetStats->ReceivedControlMessages.IsEmpty());
-		UTEST_TRUE_EXPR(BeaconHostNetStats->ReceivedControlMessages.IsEmpty());
-
-		UTEST_TRUE_EXPR(BeaconUnitTest::TickOnce(*Prerequisites));
-		UTEST_TRUE_EXPR(BeaconClient->GetConnectionState() == EBeaconConnectionState::Invalid);
-		UTEST_TRUE_EXPR(BeaconClientNetStats->ReceivedControlMessages == TArray<uint8>({ NMT_Failure }));
-		UTEST_TRUE_EXPR(BeaconHostNetStats->ReceivedControlMessages.IsEmpty());
-		UTEST_TRUE_EXPR(Prerequisites->GetStats().Client.OnFailure.InvokeCount == 1);
-		BeaconClientNetStats->ReceivedControlMessages.Empty();
-	}
-
-	// Method 3.
 	// Host handshake calls StartVerifyAuthentication (callback signature with LoginOptions) and expects callback delegate to be called to complete authentication.
 	// Test callback delegate fired with no delay is working.
 	{
@@ -2529,8 +1892,8 @@ bool FOnlineBeaconTestAuthenticatedHandshakeFailure::RunTest(const FString& Para
 		Prerequisites->GetConfig() = BaseConfig;
 		Prerequisites->GetStats() = BeaconUnitTest::FTestStats();
 
-		Prerequisites->GetConfig().Auth.Method3.bEnabled = true;
-		Prerequisites->GetConfig().Auth.Method3.Result = OnlineIdentity::Errors::InvalidCreds();
+		Prerequisites->GetConfig().Auth.bEnabled = true;
+		Prerequisites->GetConfig().Auth.Result = OnlineIdentity::Errors::InvalidCreds();
 
 		// Client setup.
 
@@ -2577,7 +1940,6 @@ bool FOnlineBeaconTestAuthenticatedHandshakeFailure::RunTest(const FString& Para
 		BeaconClientNetStats->ReceivedControlMessages.Empty();
 	}
 
-	// Method 3.
 	// Host handshake calls StartVerifyAuthentication (callback signature with LoginOptions) and expects callback delegate to be called to complete authentication.
 	// Test callback delegate fired with a single frame delay is working.
 	{
@@ -2585,9 +1947,9 @@ bool FOnlineBeaconTestAuthenticatedHandshakeFailure::RunTest(const FString& Para
 		Prerequisites->GetConfig() = BaseConfig;
 		Prerequisites->GetStats() = BeaconUnitTest::FTestStats();
 
-		Prerequisites->GetConfig().Auth.Method3.bEnabled = true;
-		Prerequisites->GetConfig().Auth.Method3.bDelayDelegate = true;
-		Prerequisites->GetConfig().Auth.Method3.Result = OnlineIdentity::Errors::InvalidCreds();
+		Prerequisites->GetConfig().Auth.bEnabled = true;
+		Prerequisites->GetConfig().Auth.bDelayDelegate = true;
+		Prerequisites->GetConfig().Auth.Result = OnlineIdentity::Errors::InvalidCreds();
 
 		// Client setup.
 
@@ -2639,15 +2001,14 @@ bool FOnlineBeaconTestAuthenticatedHandshakeFailure::RunTest(const FString& Para
 		BeaconClientNetStats->ReceivedControlMessages.Empty();
 	}
 
-	// Auth verify.
-	// Use Method 3 auth and verify that a user supplied VerifyJoinForBeaconType method does not allow the connection to proceed.
+	// Auth and verify that a user supplied VerifyJoinForBeaconType method does not allow the connection to proceed.
 	{
 		BeaconHostNetStats->ReceivedControlMessages.Empty();
 		Prerequisites->GetConfig() = BaseConfig;
 		Prerequisites->GetStats() = BeaconUnitTest::FTestStats();
 
-		Prerequisites->GetConfig().Auth.Method3.bEnabled = true;
-		Prerequisites->GetConfig().Auth.Method3.Result = FOnlineError::Success();
+		Prerequisites->GetConfig().Auth.bEnabled = true;
+		Prerequisites->GetConfig().Auth.Result = FOnlineError::Success();
 
 		Prerequisites->GetConfig().Auth.Verify.bEnabled = true;
 		Prerequisites->GetConfig().Auth.Verify.bResult = false;
diff --git a/Engine/Plugins/Online/OnlineSubsystemUtils/Source/OnlineSubsystemUtils/Private/Tests/OnlineBeaconUnitTestHost.cpp b/Engine/Plugins/Online/OnlineSubsystemUtils/Source/OnlineSubsystemUtils/Private/Tests/OnlineBeaconUnitTestHost.cpp
--- a/Engine/Plugins/Online/OnlineSubsystemUtils/Source/OnlineSubsystemUtils/Private/Tests/OnlineBeaconUnitTestHost.cpp
+++ b/Engine/Plugins/Online/OnlineSubsystemUtils/Source/OnlineSubsystemUtils/Private/Tests/OnlineBeaconUnitTestHost.cpp
@@ -22,22 +22,22 @@ bool AOnlineBeaconUnitTestHost::StartVerifyAuthentication(const FUniqueNetId& Pl
 		return false;
 	}
 
-	if (!TestConfig->Auth.Method3.bEnabled)
+	if (!TestConfig->Auth.bEnabled)
 	{
 		return Super::StartVerifyAuthentication(PlayerId, LoginOptions, AuthenticationToken, OnComplete);
 	}
 
-	if (TestConfig->Auth.Method3.bDelayDelegate)
+	if (TestConfig->Auth.bDelayDelegate)
 	{
 		return BeaconUnitTest::SetTimerForNextFrame(GetWorld(), GFrameCounter, [this, OnComplete]()
 		{
 			const BeaconUnitTest::FTestConfig* TestConfig = BeaconUnitTest::FTestPrerequisites::GetActiveTestConfig();
-			OnComplete.ExecuteIfBound(TestConfig ? TestConfig->Auth.Method3.Result : FOnlineError());
+			OnComplete.ExecuteIfBound(TestConfig ? TestConfig->Auth.Result : FOnlineError());
 		});
 	}
 	else
 	{
-		OnComplete.ExecuteIfBound(TestConfig->Auth.Method3.Result);
+		OnComplete.ExecuteIfBound(TestConfig->Auth.Result);
 		return true;
 	}
 }
diff --git a/Engine/Plugins/Online/OnlineSubsystemUtils/Source/OnlineSubsystemUtils/Public/Tests/OnlineBeaconUnitTestUtils.h b/Engine/Plugins/Online/OnlineSubsystemUtils/Source/OnlineSubsystemUtils/Public/Tests/OnlineBeaconUnitTestUtils.h
--- a/Engine/Plugins/Online/OnlineSubsystemUtils/Source/OnlineSubsystemUtils/Public/Tests/OnlineBeaconUnitTestUtils.h
+++ b/Engine/Plugins/Online/OnlineSubsystemUtils/Source/OnlineSubsystemUtils/Public/Tests/OnlineBeaconUnitTestUtils.h
@@ -179,27 +179,6 @@ namespace BeaconUnitTest
 
 		struct FAuth
 		{
-			struct FMethod1
-			{
-				bool bEnabled = false;
-				bool bDelayDelegate = false;
-				FOnlineError Result;
-			};
-
-			struct FMethod2
-			{
-				bool bEnabled = false;
-				bool bDelayDelegate = false;
-				FOnlineError Result;
-			};
-
-			struct FMethod3
-			{
-				bool bEnabled = false;
-				bool bDelayDelegate = false;
-				FOnlineError Result;
-			};
-
 			struct FVerify
 			{
 				bool bEnabled = false;
@@ -207,9 +186,8 @@ namespace BeaconUnitTest
 			};
 
 			bool bEnabled = false;
-			FMethod1 Method1;
-			FMethod2 Method2;
-			FMethod3 Method3;
+			bool bDelayDelegate = false;
+			FOnlineError Result;
 			FVerify Verify;
 		};
 
