PR #11425: Fix error linking with FVoicePacketBuffer

#11425

#rnx

diff --git a/Engine/Plugins/Online/OnlineSubsystemUtils/Source/OnlineSubsystemUtils/Public/VoicePacketBuffer.h b/Engine/Plugins/Online/OnlineSubsystemUtils/Source/OnlineSubsystemUtils/Public/VoicePacketBuffer.h
--- a/Engine/Plugins/Online/OnlineSubsystemUtils/Source/OnlineSubsystemUtils/Public/VoicePacketBuffer.h
+++ b/Engine/Plugins/Online/OnlineSubsystemUtils/Source/OnlineSubsystemUtils/Public/VoicePacketBuffer.h
@@ -10,7 +10,7 @@
 // Set this to 1 to enforce a critical section between FVoicePacketBuffer::PopAudio and FVoicePacketBuffer::SubmitPacket.
 #define SCOPELOCK_VOICE_PACKET_BUFFER 0
 
-enum class EVoipStreamDataFormat : uint8
+enum class ONLINESUBSYSTEMUTILS_API EVoipStreamDataFormat : uint8
 {
 	Float,
 	Int16,
@@ -22,7 +22,7 @@ enum class EVoipStreamDataFormat : uint8
 /* This structure represents an individual packet in the                */
 /* FVoicePacketBuffer.                                                  */
 /************************************************************************/
-struct FSortedVoicePacketNode
+struct ONLINESUBSYSTEMUTILS_API FSortedVoicePacketNode
 {
 	// The actual resulting audio data. allocated ahead of time to minimize memory churn.
 	TArray<float> AudioBufferMem;
@@ -112,7 +112,7 @@ struct FSortedVoicePacketNode
 /* This class is intended to be thread safe, as long as PushPacket() nor*/
 /* PopAudio() are called by multiple threads.                           */
 /************************************************************************/
-class FVoicePacketBuffer
+class ONLINESUBSYSTEMUTILS_API FVoicePacketBuffer
 {
 public:
 	// This is the only constructor that should be used with FVoicePacketBuffer.
