Addition of scope cycle counters to methods relevant to the hitch, to help with diagnosis.

diff --git a/Engine/Plugins/Online/OnlineFramework/Source/Party/Private/Party/SocialParty.cpp b/Engine/Plugins/Online/OnlineFramework/Source/Party/Private/Party/SocialParty.cpp
--- a/Engine/Plugins/Online/OnlineFramework/Source/Party/Private/Party/SocialParty.cpp
+++ b/Engine/Plugins/Online/OnlineFramework/Source/Party/Private/Party/SocialParty.cpp
@@ -1490,6 +1490,8 @@ void USocialParty::HandlePartyStateChanged(const FUniqueNetId& LocalUserId, cons
 
 void USocialParty::ConnectToReservationBeacon()
 {
+	QUICK_SCOPE_CYCLE_COUNTER(STAT_SocialParty_ConnectToReservationBeacon);
+
 	APartyBeaconClient* LocalReservationBeaconClient = ReservationBeaconClient.Get();
 	if (IsLocalPlayerPartyLeader() && !LocalReservationBeaconClient)
 	{
diff --git a/Engine/Plugins/Online/OnlineSubsystemUtils/Source/OnlineSubsystemUtils/Private/IpNetDriver.cpp b/Engine/Plugins/Online/OnlineSubsystemUtils/Source/OnlineSubsystemUtils/Private/IpNetDriver.cpp
--- a/Engine/Plugins/Online/OnlineSubsystemUtils/Source/OnlineSubsystemUtils/Private/IpNetDriver.cpp
+++ b/Engine/Plugins/Online/OnlineSubsystemUtils/Source/OnlineSubsystemUtils/Private/IpNetDriver.cpp
@@ -974,6 +974,8 @@ bool UIpNetDriver::InitBase( bool bInitAsClient, FNetworkNotify* InNotify, const
 
 bool UIpNetDriver::InitConnect( FNetworkNotify* InNotify, const FURL& ConnectURL, FString& Error )
 {
+	QUICK_SCOPE_CYCLE_COUNTER(STAT_IpNetDriver_InitConnect);
+
 	using namespace UE::Net::Private;
 
 	ISocketSubsystem* SocketSubsystem = GetSocketSubsystem();
diff --git a/Engine/Plugins/Online/OnlineSubsystemUtils/Source/OnlineSubsystemUtils/Private/OnlineBeaconClient.cpp b/Engine/Plugins/Online/OnlineSubsystemUtils/Source/OnlineSubsystemUtils/Private/OnlineBeaconClient.cpp
--- a/Engine/Plugins/Online/OnlineSubsystemUtils/Source/OnlineSubsystemUtils/Private/OnlineBeaconClient.cpp
+++ b/Engine/Plugins/Online/OnlineSubsystemUtils/Source/OnlineSubsystemUtils/Private/OnlineBeaconClient.cpp
@@ -145,6 +145,8 @@ FUniqueNetIdRepl GetFirstPlayerAccountId(UWorld* World)
 
 bool AOnlineBeaconClient::InitClient(FURL& URL)
 {
+	QUICK_SCOPE_CYCLE_COUNTER(STAT_OnlineBeaconClient_InitClient);
+
 	bool bSuccess = false;
 
 	if(URL.Valid)
diff --git a/Engine/Plugins/Online/OnlineSubsystemUtils/Source/OnlineSubsystemUtils/Private/PartyBeaconClient.cpp b/Engine/Plugins/Online/OnlineSubsystemUtils/Source/OnlineSubsystemUtils/Private/PartyBeaconClient.cpp
--- a/Engine/Plugins/Online/OnlineSubsystemUtils/Source/OnlineSubsystemUtils/Private/PartyBeaconClient.cpp
+++ b/Engine/Plugins/Online/OnlineSubsystemUtils/Source/OnlineSubsystemUtils/Private/PartyBeaconClient.cpp
@@ -141,6 +141,8 @@ void APartyBeaconClient::ClearTimers(bool bCallFailSafeIfNeeded)
 
 bool APartyBeaconClient::RequestReservation(const FString& ConnectInfoStr, const FString& InSessionId, const FUniqueNetIdRepl& RequestingPartyLeader, const TArray<FPlayerReservation>& PartyMembers)
 {
+	QUICK_SCOPE_CYCLE_COUNTER(STAT_PartyBeaconClient_RequestReservation_FourParams);
+
 	bool bSuccess = false;
 
 	FURL ConnectURL(NULL, *ConnectInfoStr, TRAVEL_Absolute);
@@ -169,6 +171,8 @@ bool APartyBeaconClient::RequestReservation(const FString& ConnectInfoStr, const
 
 bool APartyBeaconClient::RequestReservation(const FOnlineSessionSearchResult& DesiredHost, const FUniqueNetIdRepl& RequestingPartyLeader, const TArray<FPlayerReservation>& PartyMembers)
 {
+	QUICK_SCOPE_CYCLE_COUNTER(STAT_PartyBeaconClient_RequestReservation_ThreeParams);
+
 	bool bSuccess = false;
 
 	if (DesiredHost.IsValid())
@@ -201,6 +205,8 @@ bool APartyBeaconClient::RequestReservation(const FOnlineSessionSearchResult& De
 
 bool APartyBeaconClient::RequestReservationUpdate(const FUniqueNetIdRepl& RequestingPartyLeader, const TArray<FPlayerReservation>& PlayersToModify, bool bRemovePlayers)
 {
+	QUICK_SCOPE_CYCLE_COUNTER(STAT_PartyBeaconClient_RequestReservationUpdate_ThreeParams);
+
 	bool bWasStarted = false;
 
 	EBeaconConnectionState MyConnectionState = GetConnectionState();
@@ -219,6 +225,8 @@ bool APartyBeaconClient::RequestReservationUpdate(const FUniqueNetIdRepl& Reques
 
 bool APartyBeaconClient::RequestReservationUpdate(const FString& ConnectInfoStr, const FString& InSessionId, const FUniqueNetIdRepl& RequestingPartyLeader, const TArray<FPlayerReservation>& PlayersToModify, bool bRemovePlayers)
 {
+	QUICK_SCOPE_CYCLE_COUNTER(STAT_PartyBeaconClient_RequestReservationUpdate_FiveParams);
+
 	bool bWasStarted = false;
 
 	if (!ConnectInfoStr.IsEmpty() && !InSessionId.IsEmpty())
@@ -249,6 +257,8 @@ bool APartyBeaconClient::RequestReservationUpdate(const FString& ConnectInfoStr,
 
 bool APartyBeaconClient::RequestReservationUpdate(const FOnlineSessionSearchResult& DesiredHost, const FUniqueNetIdRepl& RequestingPartyLeader, const TArray<FPlayerReservation>& PlayersToModify, bool bRemovePlayers)
 {
+	QUICK_SCOPE_CYCLE_COUNTER(STAT_PartyBeaconClient_RequestReservationUpdate_FourParams);
+
 	bool bWasStarted = false;
 
 	EBeaconConnectionState MyConnectionState = GetConnectionState();
