Fix build error with modular build, enums cannot be exported

#rnx

diff --git a/Engine/Plugins/Online/OnlineSubsystemUtils/Source/OnlineSubsystemUtils/Public/VoicePacketBuffer.h b/Engine/Plugins/Online/OnlineSubsystemUtils/Source/OnlineSubsystemUtils/Public/VoicePacketBuffer.h
--- a/Engine/Plugins/Online/OnlineSubsystemUtils/Source/OnlineSubsystemUtils/Public/VoicePacketBuffer.h
+++ b/Engine/Plugins/Online/OnlineSubsystemUtils/Source/OnlineSubsystemUtils/Public/VoicePacketBuffer.h
@@ -10,7 +10,7 @@
 // Set this to 1 to enforce a critical section between FVoicePacketBuffer::PopAudio and FVoicePacketBuffer::SubmitPacket.
 #define SCOPELOCK_VOICE_PACKET_BUFFER 0
 
-enum class ONLINESUBSYSTEMUTILS_API EVoipStreamDataFormat : uint8
+enum class EVoipStreamDataFormat : uint8
 {
 	Float,
 	Int16,
