PR #11377

Addition of logic to pass PIE login credentials to standalone builds launched from editor

diff --git a/Engine/Plugins/Online/OnlineSubsystemUtils/Source/OnlineSubsystemUtils/Private/OnlineEngineInterfaceImpl.cpp b/Engine/Plugins/Online/OnlineSubsystemUtils/Source/OnlineSubsystemUtils/Private/OnlineEngineInterfaceImpl.cpp
--- a/Engine/Plugins/Online/OnlineSubsystemUtils/Source/OnlineSubsystemUtils/Private/OnlineEngineInterfaceImpl.cpp
+++ b/Engine/Plugins/Online/OnlineSubsystemUtils/Source/OnlineSubsystemUtils/Private/OnlineEngineInterfaceImpl.cpp
@@ -662,6 +662,17 @@ int32 UOnlineEngineInterfaceImpl::GetNumPIELogins()
 	return Online::GetUtils()->GetNumPIELogins();
 }
 
+FString UOnlineEngineInterfaceImpl::GetPIELoginCommandLineArgs(int32 Index)
+{
+	TArray<FOnlineAccountCredentials> Logins;
+	Online::GetUtils()->GetPIELogins(Logins);
+	if (Logins.IsValidIndex(Index))
+	{
+		return FString::Printf(TEXT(" -AUTH_TYPE=%s -AUTH_LOGIN=%s -AUTH_PASSWORD=%s"), *Logins[Index].Type, *Logins[Index].Id, *Logins[Index].Token);
+	}
+	return TEXT("");
+}
+
 void UOnlineEngineInterfaceImpl::SetForceDedicated(FName OnlineIdentifier, bool bForce)
 {
 	IOnlineSubsystem* OnlineSub = IOnlineSubsystem::Get(OnlineIdentifier);
diff --git a/Engine/Plugins/Online/OnlineSubsystemUtils/Source/OnlineSubsystemUtils/Private/OnlineEngineInterfaceImpl.h b/Engine/Plugins/Online/OnlineSubsystemUtils/Source/OnlineSubsystemUtils/Private/OnlineEngineInterfaceImpl.h
--- a/Engine/Plugins/Online/OnlineSubsystemUtils/Source/OnlineSubsystemUtils/Private/OnlineEngineInterfaceImpl.h
+++ b/Engine/Plugins/Online/OnlineSubsystemUtils/Source/OnlineSubsystemUtils/Private/OnlineEngineInterfaceImpl.h
@@ -158,6 +158,7 @@ public:
 	virtual bool SupportsOnlinePIE() override;
 	virtual void SetShouldTryOnlinePIE(bool bShouldTry) override;
 	virtual int32 GetNumPIELogins() override;
+	virtual FString GetPIELoginCommandLineArgs(int32 Index) override;
 	virtual void SetForceDedicated(FName OnlineIdentifier, bool bForce) override;
 	virtual void LoginPIEInstance(FName OnlineIdentifier, int32 LocalUserNum, int32 PIELoginNum, FOnPIELoginComplete& CompletionDelegate) override;
 #endif
diff --git a/Engine/Plugins/Online/OnlineSubsystemUtils/Source/OnlineSubsystemUtils/Private/OnlineServicesEngineInterfaceImpl.cpp b/Engine/Plugins/Online/OnlineSubsystemUtils/Source/OnlineSubsystemUtils/Private/OnlineServicesEngineInterfaceImpl.cpp
--- a/Engine/Plugins/Online/OnlineSubsystemUtils/Source/OnlineSubsystemUtils/Private/OnlineServicesEngineInterfaceImpl.cpp
+++ b/Engine/Plugins/Online/OnlineSubsystemUtils/Source/OnlineSubsystemUtils/Private/OnlineServicesEngineInterfaceImpl.cpp
@@ -409,6 +409,21 @@ int32 UOnlineServicesEngineInterfaceImpl::GetNumPIELogins()
 	return UE::Online::GetServicesEngineUtils()->GetNumPIELogins();
 }
 
+FString UOnlineServicesEngineInterfaceImpl::GetPIELoginCommandLineArgs(int32 Index)
+{
+	TArray<UE::Online::FAuthLogin::Params> LoginParams;
+	UE::Online::GetServicesEngineUtils()->GetPIELogins(LoginParams);
+	if (LoginParams.IsValidIndex(Index))
+	{
+		const FString& Type = LoginParams[Index].CredentialsType.ToString();
+		const FString& Id = LoginParams[Index].CredentialsId;
+		const FString& Token = LoginParams[Index].CredentialsToken.IsType<FString>() ? LoginParams[Index].CredentialsToken.Get<FString>() : LoginParams[Index].CredentialsToken.Get<UE::Online::FExternalAuthToken>().Data;
+
+		return FString::Printf(TEXT(" -AUTH_TYPE=%s -AUTH_LOGIN=%s -AUTH_PASSWORD=%s"), *Type, *Id, *Token);
+	}
+	return TEXT("");
+}
+
 void UOnlineServicesEngineInterfaceImpl::SetForceDedicated(FName OnlineIdentifier, bool bForce)
 {
 	// TODO:  Support other services? This is only called from PlayLevel with the result of GetOnlineIdentifier which only returns the default service id
diff --git a/Engine/Plugins/Online/OnlineSubsystemUtils/Source/OnlineSubsystemUtils/Private/OnlineServicesEngineInterfaceImpl.h b/Engine/Plugins/Online/OnlineSubsystemUtils/Source/OnlineSubsystemUtils/Private/OnlineServicesEngineInterfaceImpl.h
--- a/Engine/Plugins/Online/OnlineSubsystemUtils/Source/OnlineSubsystemUtils/Private/OnlineServicesEngineInterfaceImpl.h
+++ b/Engine/Plugins/Online/OnlineSubsystemUtils/Source/OnlineSubsystemUtils/Private/OnlineServicesEngineInterfaceImpl.h
@@ -123,6 +123,7 @@ public:
 	virtual bool SupportsOnlinePIE() override;
 	virtual void SetShouldTryOnlinePIE(bool bShouldTry) override;
 	virtual int32 GetNumPIELogins() override;
+	virtual FString GetPIELoginCommandLineArgs(int32 Index) override;
 	virtual void SetForceDedicated(FName OnlineIdentifier, bool bForce) override;
 	virtual void LoginPIEInstance(FName OnlineIdentifier, int32 LocalUserNum, int32 PIELoginNum, FOnPIELoginComplete& CompletionDelegate) override;
 #endif
diff --git a/Engine/Source/Editor/UnrealEd/Private/PlayLevelNewProcess.cpp b/Engine/Source/Editor/UnrealEd/Private/PlayLevelNewProcess.cpp
--- a/Engine/Source/Editor/UnrealEd/Private/PlayLevelNewProcess.cpp
+++ b/Engine/Source/Editor/UnrealEd/Private/PlayLevelNewProcess.cpp
@@ -8,6 +8,7 @@
 #include "HAL/PlatformApplicationMisc.h"
 #include "Settings/LevelEditorPlaySettings.h"
 #include "DataDrivenShaderPlatformInfo.h"
+#include "Net/OnlineEngineInterface.h"
 
 void UEditorEngine::StartPlayInNewProcessSession(FRequestPlaySessionParams& InRequestParams)
 {
@@ -297,6 +298,13 @@ void UEditorEngine::LaunchNewProcess(const FRequestPlaySessionParams& InParams,
 		{
 			CommandLine += TEXT(" -ForceRes");
 		}
+
+		// Pass through `-AUTH_LOGIN` etc arguments for external processes using the PIE credentials, if applicable.
+		const int32 NumLogins = UOnlineEngineInterface::Get()->GetNumPIELogins();
+		if (SupportsOnlinePIE() && InInstanceNum >= 0 && InInstanceNum < NumLogins)
+		{
+			CommandLine += UOnlineEngineInterface::Get()->GetPIELoginCommandLineArgs(InInstanceNum);
+		}
 	}
 
 	FString GameNameOrProjectFile;
diff --git a/Engine/Source/Runtime/Engine/Public/Net/OnlineEngineInterface.h b/Engine/Source/Runtime/Engine/Public/Net/OnlineEngineInterface.h
--- a/Engine/Source/Runtime/Engine/Public/Net/OnlineEngineInterface.h
+++ b/Engine/Source/Runtime/Engine/Public/Net/OnlineEngineInterface.h
@@ -250,6 +250,8 @@ public:
 	virtual void SetShouldTryOnlinePIE(bool bShouldTry) {}
 	/** @return number of logins stored with the editor for logging in via Play In Editor (PIE) */
 	virtual int32 GetNumPIELogins() { return 0; }
+	/** @return a string with all three components (type, id, token) for the login with the given index for logging in via Play in Editor (PIE) */
+	virtual FString GetPIELoginCommandLineArgs(int32 Index) { return TEXT(""); }
 	/** Force the online subsystem to treat itself like a dedicated server */
 	virtual void SetForceDedicated(FName OnlineIdentifier, bool bForce) {}
 	/** Login a local user to the online subsystem before starting a Play In Editor (PIE) instance, using credentials stored in the editor */
