Fix AI service spawning too many bots when bots are not allowed in human teams

diff --git a/Engine/Plugins/Online/OnlineSubsystemUtils/Source/OnlineSubsystemUtils/Private/PartyBeaconState.cpp b/Engine/Plugins/Online/OnlineSubsystemUtils/Source/OnlineSubsystemUtils/Private/PartyBeaconState.cpp
--- a/Engine/Plugins/Online/OnlineSubsystemUtils/Source/OnlineSubsystemUtils/Private/PartyBeaconState.cpp
+++ b/Engine/Plugins/Online/OnlineSubsystemUtils/Source/OnlineSubsystemUtils/Private/PartyBeaconState.cpp
@@ -293,6 +293,21 @@ int32 UPartyBeaconState::GetTeamForCurrentPlayer(const FUniqueNetId& PlayerId) c
 	return TeamNum;
 }
 
+int32 UPartyBeaconState::GetReservationTeamCount() const
+{
+	TSet<int32> Teams; 
+	for (int32 ResIdx = 0; ResIdx < Reservations.Num(); ResIdx++)
+	{
+		const FPartyReservation& Reservation = Reservations[ResIdx];
+		if (Reservation.TeamNum != INDEX_NONE)
+		{
+			Teams.Add(Reservation.TeamNum);
+		}
+	}
+
+	return Teams.Num();
+}
+
 int32 UPartyBeaconState::GetPlayersOnTeam(int32 TeamIndex, TArray<FUniqueNetIdRepl>& TeamMembers) const
 {
 	TeamMembers.Empty(NumPlayersPerTeam);
diff --git a/Engine/Plugins/Online/OnlineSubsystemUtils/Source/OnlineSubsystemUtils/Public/PartyBeaconState.h b/Engine/Plugins/Online/OnlineSubsystemUtils/Source/OnlineSubsystemUtils/Public/PartyBeaconState.h
--- a/Engine/Plugins/Online/OnlineSubsystemUtils/Source/OnlineSubsystemUtils/Public/PartyBeaconState.h
+++ b/Engine/Plugins/Online/OnlineSubsystemUtils/Source/OnlineSubsystemUtils/Public/PartyBeaconState.h
@@ -465,6 +465,13 @@ class ONLINESUBSYSTEMUTILS_API UPartyBeaconState : public UObject
 	 */
 	virtual int32 GetTeamForCurrentPlayer(const FUniqueNetId& PlayerId) const;
 
+	/**
+	 * Get the amount of teams used by reservations
+	 *
+	 * @return The number of teams that are used by reservations
+	 */
+	virtual int32 GetReservationTeamCount() const;
+
 	/**
 	 * Get all the known players on a given team
 	 * 
