Allow party members to update reservations.

diff --git a/Engine/Plugins/Online/OnlineFramework/Source/Party/Private/Party/SocialParty.cpp b/Engine/Plugins/Online/OnlineFramework/Source/Party/Private/Party/SocialParty.cpp
--- a/Engine/Plugins/Online/OnlineFramework/Source/Party/Private/Party/SocialParty.cpp
+++ b/Engine/Plugins/Online/OnlineFramework/Source/Party/Private/Party/SocialParty.cpp
@@ -479,6 +479,11 @@ void USocialParty::InitializeParty(const TSharedRef<const FOnlineParty>& InOssPa
 	}
 }
 
+bool USocialParty::AllowJoinInProgressToMember()
+{
+	return false;
+}
+
 void USocialParty::InitializePartyInternal()
 {
 	QUICK_SCOPE_CYCLE_COUNTER(STAT_SocialParty_InitializePartyInternal);
@@ -1492,8 +1497,10 @@ void USocialParty::ConnectToReservationBeacon()
 {
 	QUICK_SCOPE_CYCLE_COUNTER(STAT_SocialParty_ConnectToReservationBeacon);
 
+	const bool bJIPRequiresLeaderApproval = !AllowJoinInProgressToMember();
+
 	APartyBeaconClient* LocalReservationBeaconClient = ReservationBeaconClient.Get();
-	if (IsLocalPlayerPartyLeader() && !LocalReservationBeaconClient)
+	if ((IsLocalPlayerPartyLeader() || !bJIPRequiresLeaderApproval) && !LocalReservationBeaconClient)
 	{
 		FPendingMemberApproval NextApproval;
 		if (PendingApprovals.Peek(NextApproval))
@@ -1547,7 +1554,8 @@ void USocialParty::ConnectToReservationBeacon()
 									}
 								}
 
-								bStartedConnection = ReservationBeaconClient->RequestReservationUpdate(URL, Session->GetSessionIdStr(), GetPartyLeader()->GetPrimaryNetId(), ReservationAsArray, NextApproval.bIsPlayerRemoval);
+								const FUniqueNetIdRepl ReservationLeaderId = AllowJoinInProgressToMember() ? OwningLocalUserId : GetPartyLeader()->GetPrimaryNetId();
+								bStartedConnection = ReservationBeaconClient->RequestReservationUpdate(URL, Session->GetSessionIdStr(), ReservationLeaderId, ReservationAsArray, NextApproval.bIsPlayerRemoval);
 							}
 						}
 					}
@@ -1648,7 +1656,9 @@ void USocialParty::PumpApprovalQueue()
 					NewPlayerRes.bAllowCrossplay = true;
 				}
 			}
-			LocalReservationBeaconClient->RequestReservationUpdate(GetPartyLeader()->GetPrimaryNetId(), PlayersToAdd);
+
+			const FUniqueNetIdRepl ReservationLeaderId = AllowJoinInProgressToMember() ? OwningLocalUserId : GetPartyLeader()->GetPrimaryNetId();
+			LocalReservationBeaconClient->RequestReservationUpdate(ReservationLeaderId, PlayersToAdd);
 		}
 		else
 		{
@@ -1953,14 +1963,14 @@ void USocialParty::RequestJoinInProgress(const UPartyMember& TargetMember, const
 		if (RequestingMember && ensure(RequestingMember->IsLocalPlayer()))
 		{
 			RequestingMember->GetMutableRepData().SetJoinInProgressDataRequest(Request);
-		}	
+		}
 	}
 	RunJoinInProgressTimer();
 }
 
 void USocialParty::HandleJoinInProgressDataRequestChanged(const FPartyMemberJoinInProgressRequest& Request, UPartyMember* Member)
 {
-	if (Request.Time == 0 || !IsLocalPlayerPartyLeader())
+	if (Request.Time == 0 || !(AllowJoinInProgressToMember() || IsLocalPlayerPartyLeader()))
 	{
 		// Ignore if this is not an active request or if we're not the party leader.
 		return;
diff --git a/Engine/Plugins/Online/OnlineFramework/Source/Party/Public/Party/SocialParty.h b/Engine/Plugins/Online/OnlineFramework/Source/Party/Public/Party/SocialParty.h
--- a/Engine/Plugins/Online/OnlineFramework/Source/Party/Public/Party/SocialParty.h
+++ b/Engine/Plugins/Online/OnlineFramework/Source/Party/Public/Party/SocialParty.h
@@ -256,6 +256,8 @@ protected:
 	bool TryInviteUser(const USocialUser& UserToInvite, const ESocialPartyInviteMethod InviteMethod = ESocialPartyInviteMethod::Other, const FString& MetaData = FString());
 	//--------------------------
 
+	virtual bool AllowJoinInProgressToMember();
+
 protected:
 	virtual void InitializePartyInternal();
 
diff --git a/Engine/Plugins/Online/OnlineSubsystemUtils/Source/OnlineSubsystemUtils/Private/PartyBeaconHost.cpp b/Engine/Plugins/Online/OnlineSubsystemUtils/Source/OnlineSubsystemUtils/Private/PartyBeaconHost.cpp
--- a/Engine/Plugins/Online/OnlineSubsystemUtils/Source/OnlineSubsystemUtils/Private/PartyBeaconHost.cpp
+++ b/Engine/Plugins/Online/OnlineSubsystemUtils/Source/OnlineSubsystemUtils/Private/PartyBeaconHost.cpp
@@ -11,6 +11,16 @@
 
 #include UE_INLINE_GENERATED_CPP_BY_NAME(PartyBeaconHost)
 
+namespace BeaconConsoleVariables
+{
+TAutoConsoleVariable<bool> CVarRemapReservationUpdateToExistingReservation(
+	TEXT("beacon.RemapReservationUpdateToExistingReservation"),
+	true,
+	TEXT("Whether to remap the leader of a reservation update request to the leader of the reservation they appear in.\n")
+	TEXT("Enabled"),
+	ECVF_Default);
+}
+
 APartyBeaconHost::APartyBeaconHost(const FObjectInitializer& ObjectInitializer) :
 	Super(ObjectInitializer),
 	State(NULL),
@@ -1172,10 +1182,35 @@ void APartyBeaconHost::ProcessReservationUpdateRequest(APartyBeaconClient* Clien
 
 	if (Client && Client->GetUniqueId().IsValid() && Client->GetUniqueId() == ReservationUpdateRequest.PartyLeader)
 	{
+		const FPartyReservation* EffectiveReservationUpdateRequest = &ReservationUpdateRequest;
+		FPartyReservation RemappedReservationUpdateRequest;
+		if (BeaconConsoleVariables::CVarRemapReservationUpdateToExistingReservation.GetValueOnGameThread())
+		{
+			const int32 ExistingReservationIdx = State ? State->GetExistingReservationContainingMember(ReservationUpdateRequest.PartyLeader) : INDEX_NONE;
+			const FPartyReservation* ExistingReservation = ExistingReservationIdx != INDEX_NONE ? &State->GetReservations()[ExistingReservationIdx] : nullptr;
+			if (ExistingReservation)
+			{
+				if (ExistingReservation->PartyLeader != RemappedReservationUpdateRequest.PartyLeader)
+				{
+					UE_LOG(LogPartyBeacon, Log, TEXT("Remapping reservation update request leader [%s] to leader found in existing reservation [%s]"),
+						*ExistingReservation->PartyLeader.ToDebugString(),
+						*RemappedReservationUpdateRequest.PartyLeader.ToDebugString());
+					RemappedReservationUpdateRequest = ReservationUpdateRequest;
+					RemappedReservationUpdateRequest.PartyLeader = ExistingReservation->PartyLeader;
+					EffectiveReservationUpdateRequest = &RemappedReservationUpdateRequest;
+				}
+			}
+			else
+			{
+				// Invalid request, member has no existing reservation to update.
+				// Continue with call to UpdatePartyReservation which will return an error.
+			}
+		}
+
 		EPartyReservationResult::Type Result = EPartyReservationResult::BadSessionId;
 		if (DoesSessionMatch(SessionId))
 		{
-			Result = UpdatePartyReservation(ReservationUpdateRequest, bIsRemovingMember);
+			Result = UpdatePartyReservation(*EffectiveReservationUpdateRequest, bIsRemovingMember);
 		}
 
 		UE_LOG(LogPartyBeacon, Verbose, TEXT("ProcessReservationUpdateRequest result: %s"), EPartyReservationResult::ToString(Result));
@@ -1183,7 +1218,7 @@ void APartyBeaconHost::ProcessReservationUpdateRequest(APartyBeaconClient* Clien
 			(Result != EPartyReservationResult::ReservationAccepted))
 		{
 			DumpReservations();
-			ReservationUpdateRequest.Dump();
+			EffectiveReservationUpdateRequest->Dump();
 		}
 
 		Client->ClientReservationResponse(Result);
