Updated OnlinePurchaseONEStore IOnlinePurchase interface

diff --git a/Engine/Plugins/Online/OnlineSubsystemUtils/Source/OnlineSubsystemUtils/Private/OnlineSubsystemUtilsModule.cpp b/Engine/Plugins/Online/OnlineSubsystemUtils/Source/OnlineSubsystemUtils/Private/OnlineSubsystemUtilsModule.cpp
--- a/Engine/Plugins/Online/OnlineSubsystemUtils/Source/OnlineSubsystemUtils/Private/OnlineSubsystemUtilsModule.cpp
+++ b/Engine/Plugins/Online/OnlineSubsystemUtils/Source/OnlineSubsystemUtils/Private/OnlineSubsystemUtilsModule.cpp
@@ -238,6 +238,7 @@ private:
 		CREATE_HASH(AMAZON_SUBSYSTEM);
 		CREATE_HASH(GAMECIRCLE_SUBSYSTEM);
 		CREATE_HASH(THUNDERHEAD_SUBSYSTEM); 
+		CREATE_HASH(ONESTORE_SUBSYSTEM);
 		// Shouldn't need these as they are mocking interfaces for existing platforms
 		CREATE_HASH(PS4SERVER_SUBSYSTEM);
 
