Fixed a crash ensure that occurred when stopping play in editor while a voice chat session was active.

diff --git a/Engine/Plugins/Online/OnlineSubsystemUtils/Source/OnlineSubsystemUtils/Private/OnlineSubsystemUtils.cpp b/Engine/Plugins/Online/OnlineSubsystemUtils/Source/OnlineSubsystemUtils/Private/OnlineSubsystemUtils.cpp
--- a/Engine/Plugins/Online/OnlineSubsystemUtils/Source/OnlineSubsystemUtils/Private/OnlineSubsystemUtils.cpp
+++ b/Engine/Plugins/Online/OnlineSubsystemUtils/Source/OnlineSubsystemUtils/Private/OnlineSubsystemUtils.cpp
@@ -228,9 +228,12 @@ UWorld* GetWorldForOnline(FName InstanceName)
 #if WITH_EDITOR
 	if (InstanceName != FOnlineSubsystemImpl::DefaultInstanceName && InstanceName != NAME_None)
 	{
-		FWorldContext& WorldContext = GEngine->GetWorldContextFromHandleChecked(InstanceName);
-		check(WorldContext.WorldType == EWorldType::Game || WorldContext.WorldType == EWorldType::PIE);
-		World = WorldContext.World();
+		FWorldContext* WorldContext = GEngine->GetWorldContextFromHandle(InstanceName);
+		if (WorldContext)
+		{
+			check(WorldContext->WorldType == EWorldType::Game || WorldContext->WorldType == EWorldType::PIE);
+			World = WorldContext->World();
+		}
 	}
 	else
 #endif
