Added better error context to OnlineBeacon failures

#Synced-CL 49316295

diff --git a/Engine/Plugins/Online/OnlineFramework/Source/Lobby/Private/LobbyBeaconClient.cpp b/Engine/Plugins/Online/OnlineFramework/Source/Lobby/Private/LobbyBeaconClient.cpp
--- a/Engine/Plugins/Online/OnlineFramework/Source/Lobby/Private/LobbyBeaconClient.cpp
+++ b/Engine/Plugins/Online/OnlineFramework/Source/Lobby/Private/LobbyBeaconClient.cpp
@@ -87,7 +87,7 @@ void ALobbyBeaconClient::ConnectToLobby(const FOnlineSessionSearchResult& Desire
 
 	if (!bSuccess)
 	{
-		OnFailure();
+		AOnlineBeacon::OnFailure(EBeaconFailureReason::InitFailure, TEXTVIEW("beacon_lobbyclient_init_failure"));
 	}
 }
 
diff --git a/Engine/Plugins/Online/OnlineFramework/Source/Qos/Private/QosBeaconClient.cpp b/Engine/Plugins/Online/OnlineFramework/Source/Qos/Private/QosBeaconClient.cpp
--- a/Engine/Plugins/Online/OnlineFramework/Source/Qos/Private/QosBeaconClient.cpp
+++ b/Engine/Plugins/Online/OnlineFramework/Source/Qos/Private/QosBeaconClient.cpp
@@ -56,7 +56,7 @@ void AQosBeaconClient::SendQosRequest(const FOnlineSessionSearchResult& DesiredH
 
 	if (!bSuccess)
 	{
-		OnFailure();
+		AOnlineBeacon::OnFailure(EBeaconFailureReason::InitFailure, TEXTVIEW("beacon_qosclient_init_failure"));
 	}
 }
 
diff --git a/Engine/Plugins/Online/OnlineSubsystemUtils/Source/OnlineSubsystemUtils/Private/OnlineBeacon.cpp b/Engine/Plugins/Online/OnlineSubsystemUtils/Source/OnlineSubsystemUtils/Private/OnlineBeacon.cpp
--- a/Engine/Plugins/Online/OnlineSubsystemUtils/Source/OnlineSubsystemUtils/Private/OnlineBeacon.cpp
+++ b/Engine/Plugins/Online/OnlineSubsystemUtils/Source/OnlineSubsystemUtils/Private/OnlineBeacon.cpp
@@ -81,7 +81,7 @@ void AOnlineBeacon::HandleNetworkFailure(UWorld *World, UNetDriver *InNetDriver,
 	if (InNetDriver && InNetDriver->NetDriverName == NetDriverName)
 	{
 		UE_LOG(LogBeacon, Verbose, TEXT("NetworkFailure %s: %s"), *GetName(), ENetworkFailure::ToString(FailureType));
-		OnFailure();
+		OnFailure(EBeaconFailureReason::TransportError, TEXTVIEW("beacon_network_failure"));
 	}
 }
 
@@ -91,6 +91,14 @@ void AOnlineBeacon::OnFailure()
 	CleanupNetDriver();
 }
 
+void AOnlineBeacon::OnFailure(EBeaconFailureReason Reason, FStringView ErrorMessage)
+{
+	UE_LOG(LogBeacon, Verbose, TEXT("Online Beacon Failure: %s. %s"), LexToString(Reason), ErrorMessage.IsEmpty() ? TEXT("") : ErrorMessage.GetData());
+	PRAGMA_DISABLE_DEPRECATION_WARNINGS
+	OnFailure();
+	PRAGMA_ENABLE_DEPRECATION_WARNINGS
+}
+
 void AOnlineBeacon::OnActorChannelOpen(FInBunch& Bunch, UNetConnection* Connection)
 {
 	Connection->OwningActor = this;
diff --git a/Engine/Plugins/Online/OnlineSubsystemUtils/Source/OnlineSubsystemUtils/Private/OnlineBeaconClient.cpp b/Engine/Plugins/Online/OnlineSubsystemUtils/Source/OnlineSubsystemUtils/Private/OnlineBeaconClient.cpp
--- a/Engine/Plugins/Online/OnlineSubsystemUtils/Source/OnlineSubsystemUtils/Private/OnlineBeaconClient.cpp
+++ b/Engine/Plugins/Online/OnlineSubsystemUtils/Source/OnlineSubsystemUtils/Private/OnlineBeaconClient.cpp
@@ -217,7 +217,7 @@ bool AOnlineBeaconClient::InitClient(const FOnlineBeaconClientInitParams& Params
 						if (GetConnectionState() == EBeaconConnectionState::Pending)
 						{
 							SetConnectionState(EBeaconConnectionState::Invalid);
-							OnFailure();
+							AOnlineBeacon::OnFailure(EBeaconFailureReason::InitFailure, TEXTVIEW("beacon_client_init_handshake_failure"));
 						}
 					}
 				}
@@ -226,7 +226,7 @@ bool AOnlineBeaconClient::InitClient(const FOnlineBeaconClientInitParams& Params
 					// error initializing the network stack, no world...
 					UE_LOG(LogBeacon, Log, TEXT("AOnlineBeaconClient::InitClient failed - no world"));
 					SetConnectionState(EBeaconConnectionState::Invalid);
-					OnFailure();
+					AOnlineBeacon::OnFailure(EBeaconFailureReason::InitFailure, TEXTVIEW("beacon_client_init_no_world"));
 				}
 			}
 			else
@@ -234,7 +234,7 @@ bool AOnlineBeaconClient::InitClient(const FOnlineBeaconClientInitParams& Params
 				// error initializing the network stack...
 				UE_LOG(LogBeacon, Log, TEXT("AOnlineBeaconClient::InitClient failed"));
 				SetConnectionState(EBeaconConnectionState::Invalid);
-				OnFailure();
+				AOnlineBeacon::OnFailure(EBeaconFailureReason::InitFailure, TEXTVIEW("beacon_client_init_failure"));
 			}
 		}
 	}
@@ -259,7 +259,7 @@ void AOnlineBeaconClient::Tick(float DeltaTime)
 			(NetDriver->ServerConnection->GetConnectionState() == USOCK_Closed))
 		{
 			UE_LOG(LogBeacon, Verbose, TEXT("Client beacon (%s) socket has closed, triggering failure."), *GetName());
-			OnFailure();
+			AOnlineBeacon::OnFailure(EBeaconFailureReason::TransportError, TEXTVIEW("beacon_client_connection_closed_remotely"));
 		}
 	}
 }
@@ -317,7 +317,7 @@ void AOnlineBeaconClient::SendInitialJoin()
 			UE_LOG(LogNet, Error, TEXT("AOnlineBeaconClient::SendInitialJoin: EncryptionToken is empty when 'net.AllowEncryption' requires it."));
 
 			SetConnectionState(EBeaconConnectionState::Invalid);
-			OnFailure();
+			Super::OnFailure(EBeaconFailureReason::EncryptionFailure, TEXTVIEW("beacon_client_empty_encryption_token"));
 		}
 	}
 }
@@ -327,7 +327,9 @@ void AOnlineBeaconClient::OnFailure()
 	UE_LOG(LogBeacon, Verbose, TEXT("Client beacon (%s) connection failure, handling connection timeout."), *GetName());
 	SetConnectionState(EBeaconConnectionState::Invalid);
 	HostConnectionFailure.ExecuteIfBound();
+	PRAGMA_DISABLE_DEPRECATION_WARNINGS
 	Super::OnFailure();
+	PRAGMA_ENABLE_DEPRECATION_WARNINGS
 }
 
 
@@ -423,7 +425,7 @@ void AOnlineBeaconClient::NotifyControlMessage(UNetConnection* Connection, uint8
 					{
 						// Force close the session
 						UE_LOG(LogBeacon, Warning, TEXT("%s: No delegate available to handle encryption ack, disconnecting."), *Connection->GetName());
-						OnFailure();
+						AOnlineBeacon::OnFailure(EBeaconFailureReason::EncryptionFailure, TEXTVIEW("beacon_client_encryption_delegate_missing"));
 					}
 				}
 				break;
@@ -462,7 +464,7 @@ void AOnlineBeaconClient::NotifyControlMessage(UNetConnection* Connection, uint8
 			{
 				// Force close the session
 				UE_LOG(LogBeacon, Warning, TEXT("%s: Unable to parse challenge request message."), *Connection->GetName());
-				OnFailure();
+				AOnlineBeacon::OnFailure(EBeaconFailureReason::TransportError, TEXTVIEW("beacon_client_parse_challenge_error"));
 			}
 			break;
 		}
@@ -481,7 +483,7 @@ void AOnlineBeaconClient::NotifyControlMessage(UNetConnection* Connection, uint8
 				{
 					// Force close the session
 					UE_LOG(LogBeacon, Log, TEXT("Beacon close from invalid beacon type"));
-					OnFailure();
+					AOnlineBeacon::OnFailure(EBeaconFailureReason::VersionMismatch, TEXTVIEW("beacon_client_invalid_beacon_type"));
 				}
 				break;
 			}
@@ -505,7 +507,12 @@ void AOnlineBeaconClient::NotifyControlMessage(UNetConnection* Connection, uint8
 						{
 							float Timeout = Connection->GetTimeoutValue();
 							UE_LOG(LogBeacon, Verbose, TEXT("%s: Waiting for first RPC from server, timeout: %.3f"), *GetName(), Timeout);
-							FTimerDelegate TimerDelegate = FTimerDelegate::CreateUObject(this, &AOnlineBeaconClient::OnFailure);
+							FTimerDelegate TimerDelegate;
+							TimerDelegate.BindWeakLambda(this, [this]() mutable
+								{
+									this->AOnlineBeacon::OnFailure(EBeaconFailureReason::TimeOut, TEXTVIEW("beacon_client_rpc_timeout"));
+								});
+							
 							GetWorldTimerManager().SetTimer(TimerHandle_OnFailure, TimerDelegate, Timeout, false);
 						}
 					}
@@ -513,7 +520,7 @@ void AOnlineBeaconClient::NotifyControlMessage(UNetConnection* Connection, uint8
 					{
 						// Force close the session
 						UE_LOG(LogBeacon, Log, TEXT("Beacon close from invalid NetGUID"));
-						OnFailure();
+						AOnlineBeacon::OnFailure(EBeaconFailureReason::TransportError, TEXTVIEW("beacon_client_invalid_netguid"));
 					}
 				}
 
@@ -555,7 +562,7 @@ void AOnlineBeaconClient::NotifyControlMessage(UNetConnection* Connection, uint8
 
 					// Force close the session
 					UE_LOG(LogBeacon, Log, TEXT("Beacon close from NMT_Failure %s"), *ErrorMsg);
-					OnFailure();
+					AOnlineBeacon::OnFailure(EBeaconFailureReason::TransportError, TEXTVIEW("beacon_client_nmt_failure"));
 				}
 
 				break;
@@ -566,7 +573,7 @@ void AOnlineBeaconClient::NotifyControlMessage(UNetConnection* Connection, uint8
 			{
 				// Force close the session
 				UE_LOG(LogBeacon, Log, TEXT("Beacon close from unexpected control message"));
-				OnFailure();
+				AOnlineBeacon::OnFailure(EBeaconFailureReason::TransportError, TEXTVIEW("beacon_client_unexpected_control_message"));
 				break;
 			}
 		}
@@ -605,19 +612,19 @@ void AOnlineBeaconClient::FinalizeEncryptedConnection(const FEncryptionKeyRespon
 			{
 				FString ResponseStr(LexToString(Response.Response));
 				UE_LOG(LogBeacon, Warning, TEXT("AOnlineBeaconClient::FinalizeEncryptedConnection: encryption failure [%s] %s"), *ResponseStr, *Response.ErrorMsg);
-				OnFailure();
+				AOnlineBeacon::OnFailure(EBeaconFailureReason::EncryptionFailure, TEXTVIEW("beacon_client_encryption_failure"));
 			}
 		}
 		else
 		{
 			UE_LOG(LogBeacon, Warning, TEXT("AOnlineBeaconClient::FinalizeEncryptedConnection: connection in invalid state. %s"), *Connection->Describe());
-			OnFailure();
+			AOnlineBeacon::OnFailure(EBeaconFailureReason::TransportError, TEXTVIEW("beacon_client_invalid_state_connection"));
 		}
 	}
 	else
 	{
 		UE_LOG(LogBeacon, Warning, TEXT("AOnlineBeaconClient::FinalizeEncryptedConnection: Connection is null."));
-		OnFailure();
+		AOnlineBeacon::OnFailure(EBeaconFailureReason::TransportError, TEXTVIEW("beacon_client_null_connection"));
 	}
 }
 
diff --git a/Engine/Plugins/Online/OnlineSubsystemUtils/Source/OnlineSubsystemUtils/Private/OnlineBeaconHost.cpp b/Engine/Plugins/Online/OnlineSubsystemUtils/Source/OnlineSubsystemUtils/Private/OnlineBeaconHost.cpp
--- a/Engine/Plugins/Online/OnlineSubsystemUtils/Source/OnlineSubsystemUtils/Private/OnlineBeaconHost.cpp
+++ b/Engine/Plugins/Online/OnlineSubsystemUtils/Source/OnlineSubsystemUtils/Private/OnlineBeaconHost.cpp
@@ -117,7 +117,7 @@ bool AOnlineBeaconHost::InitHost()
 			{
 				// error initializing the network stack...
 				UE_LOG(LogBeacon, Log, TEXT("%s: AOnlineBeaconHost::InitHost failed"), *GetName());
-				OnFailure();
+				OnFailure(EBeaconFailureReason::InitFailure, TEXTVIEW("beacon_host_init_failure"));
 			}
 		}
 	}
diff --git a/Engine/Plugins/Online/OnlineSubsystemUtils/Source/OnlineSubsystemUtils/Private/PartyBeaconClient.cpp b/Engine/Plugins/Online/OnlineSubsystemUtils/Source/OnlineSubsystemUtils/Private/PartyBeaconClient.cpp
--- a/Engine/Plugins/Online/OnlineSubsystemUtils/Source/OnlineSubsystemUtils/Private/PartyBeaconClient.cpp
+++ b/Engine/Plugins/Online/OnlineSubsystemUtils/Source/OnlineSubsystemUtils/Private/PartyBeaconClient.cpp
@@ -163,7 +163,7 @@ bool APartyBeaconClient::RequestReservation(const FString& ConnectInfoStr, const
 
 	if (!bSuccess)
 	{
-		OnFailure();
+		AOnlineBeacon::OnFailure(EBeaconFailureReason::InitFailure, TEXTVIEW("beacon_party_init_failure"));
 	}
 
 	return bSuccess;
@@ -173,8 +173,6 @@ bool APartyBeaconClient::RequestReservation(const FOnlineSessionSearchResult& De
 {
 	QUICK_SCOPE_CYCLE_COUNTER(STAT_PartyBeaconClient_RequestReservation_ThreeParams);
 
-	bool bSuccess = false;
-
 	if (DesiredHost.IsValid())
 	{
 		UWorld* World = GetWorld();
@@ -195,12 +193,10 @@ bool APartyBeaconClient::RequestReservation(const FOnlineSessionSearchResult& De
 		}
 	}
 
-	if (!bSuccess)
-	{
-		OnFailure();
-	}
+	
+	AOnlineBeacon::OnFailure(EBeaconFailureReason::ApplicationError, TEXTVIEW("beacon_party_request_reservation_failure"));
 
-	return bSuccess;
+	return false;
 }
 
 bool APartyBeaconClient::RequestReservationUpdate(const FUniqueNetIdRepl& RequestingPartyLeader, const TArray<FPlayerReservation>& PlayersToModify, bool bRemovePlayers)
@@ -302,7 +298,7 @@ bool APartyBeaconClient::RequestAddOrUpdateReservation(const FString& ConnectInf
 
 	if (!bSuccess)
 	{
-		OnFailure();
+		AOnlineBeacon::OnFailure(EBeaconFailureReason::ApplicationError, TEXTVIEW("beacon_party_update_reservation_failure"));
 	}
 
 	return bSuccess;
@@ -374,7 +370,7 @@ void APartyBeaconClient::OnConnected()
 		else
 		{
 			UE_LOG(LogPartyBeacon, Warning, TEXT("Failed to handle reservation request type %s"), ToString(RequestType));
-			OnFailure();
+			AOnlineBeacon::OnFailure(EBeaconFailureReason::TransportError, TEXTVIEW("beacon_party_reservation_failure"));
 		}
 	}
 	else
diff --git a/Engine/Plugins/Online/OnlineSubsystemUtils/Source/OnlineSubsystemUtils/Private/SpectatorBeaconClient.cpp b/Engine/Plugins/Online/OnlineSubsystemUtils/Source/OnlineSubsystemUtils/Private/SpectatorBeaconClient.cpp
--- a/Engine/Plugins/Online/OnlineSubsystemUtils/Source/OnlineSubsystemUtils/Private/SpectatorBeaconClient.cpp
+++ b/Engine/Plugins/Online/OnlineSubsystemUtils/Source/OnlineSubsystemUtils/Private/SpectatorBeaconClient.cpp
@@ -144,7 +144,7 @@ bool ASpectatorBeaconClient::RequestReservation(const FString& ConnectInfoStr, c
 
 	if (!bSuccess)
 	{
-		OnFailure();
+		AOnlineBeacon::OnFailure(EBeaconFailureReason::InitFailure, TEXTVIEW("beacon_spectator_init_failure"));
 	}
 
 	return bSuccess;
@@ -152,7 +152,6 @@ bool ASpectatorBeaconClient::RequestReservation(const FString& ConnectInfoStr, c
 
 bool ASpectatorBeaconClient::RequestReservation(const FOnlineSessionSearchResult& DesiredHost, const FUniqueNetIdRepl& RequestingSpectator, const FPlayerReservation& Spectator)
 {
-	bool bSuccess = false;
 
 	if (DesiredHost.IsValid())
 	{
@@ -174,12 +173,9 @@ bool ASpectatorBeaconClient::RequestReservation(const FOnlineSessionSearchResult
 		}
 	}
 
-	if (!bSuccess)
-	{
-		OnFailure();
-	}
-
-	return bSuccess;
+	AOnlineBeacon::OnFailure(EBeaconFailureReason::ApplicationError, TEXTVIEW("beacon_spectator_request_reservation_failure"));
+	
+	return false;
 }
 
 void ASpectatorBeaconClient::CancelReservation()
@@ -230,7 +226,7 @@ void ASpectatorBeaconClient::OnConnected()
 		else
 		{
 			UE_LOG(LogSpectatorBeacon, Warning, TEXT("Failed to handle reservation request type %s"), ToString(RequestType));
-			OnFailure();
+			AOnlineBeacon::OnFailure(EBeaconFailureReason::TransportError, TEXTVIEW("beacon_spectator_reservation_failure"));
 		}
 	}
 	else
diff --git a/Engine/Plugins/Online/OnlineSubsystemUtils/Source/OnlineSubsystemUtils/Private/Tests/OnlineBeaconUnitTestHost.cpp b/Engine/Plugins/Online/OnlineSubsystemUtils/Source/OnlineSubsystemUtils/Private/Tests/OnlineBeaconUnitTestHost.cpp
--- a/Engine/Plugins/Online/OnlineSubsystemUtils/Source/OnlineSubsystemUtils/Private/Tests/OnlineBeaconUnitTestHost.cpp
+++ b/Engine/Plugins/Online/OnlineSubsystemUtils/Source/OnlineSubsystemUtils/Private/Tests/OnlineBeaconUnitTestHost.cpp
@@ -64,8 +64,10 @@ void AOnlineBeaconUnitTestHost::OnFailure()
 	{
 		++TestStats->Host.OnFailure.InvokeCount;
 	}
-
+	
+	PRAGMA_DISABLE_DEPRECATION_WARNINGS
 	Super::OnFailure();
+	PRAGMA_ENABLE_DEPRECATION_WARNINGS
 }
 
 #endif /* WITH_DEV_AUTOMATION_TESTS && WITH_EDITOR */
diff --git a/Engine/Plugins/Online/OnlineSubsystemUtils/Source/OnlineSubsystemUtils/Public/OnlineBeacon.h b/Engine/Plugins/Online/OnlineSubsystemUtils/Source/OnlineSubsystemUtils/Public/OnlineBeacon.h
--- a/Engine/Plugins/Online/OnlineSubsystemUtils/Source/OnlineSubsystemUtils/Public/OnlineBeacon.h
+++ b/Engine/Plugins/Online/OnlineSubsystemUtils/Source/OnlineSubsystemUtils/Public/OnlineBeacon.h
@@ -26,6 +26,40 @@ namespace EBeaconState
 	};
 }
 
+enum class EBeaconFailureReason : uint8
+{
+	Unknown = 0,
+	TransportError = 1,   
+	VersionMismatch = 2,
+	InitFailure = 3,
+	EncryptionFailure = 4,
+	ApplicationError = 5,
+	TimeOut
+};
+
+inline const TCHAR* LexToString(EBeaconFailureReason Value)
+{
+	switch (Value)
+	{
+	case EBeaconFailureReason::Unknown:
+		return TEXT("Unknown");
+	case EBeaconFailureReason::TransportError:
+		return TEXT("TransportError");
+	case EBeaconFailureReason::VersionMismatch:
+		return TEXT("VersionMismatch");
+	case EBeaconFailureReason::InitFailure:
+		return TEXT("FailedToInitBeacon");
+	case EBeaconFailureReason::EncryptionFailure:
+		return TEXT("EncryptionFailure");
+	case EBeaconFailureReason::ApplicationError:
+		return TEXT("ApplicationError");
+	case EBeaconFailureReason::TimeOut:
+		return TEXT("TimeOut");
+	default:
+		return TEXT("Unknown");
+	}
+}
+
 /**
  * Base class for beacon communication (Unreal Networking, but outside normal gameplay traffic)
  */
@@ -126,8 +160,11 @@ protected:
 	ONLINESUBSYSTEMUTILS_API virtual bool InitBase();
 
 	/** Notification that failure needs to be handled */
+	UE_DEPRECATED(5.7, "Deprecated - Move to OnFailure(FBeaconFailureContext).")
 	ONLINESUBSYSTEMUTILS_API virtual void OnFailure();
 
+	ONLINESUBSYSTEMUTILS_API virtual void OnFailure(EBeaconFailureReason Reason, FStringView ErrorMessage);
+
 	/** overridden to return that player controllers are capable of RPCs */
 	ONLINESUBSYSTEMUTILS_API virtual bool HasNetOwner() const override;
 };
diff --git a/Engine/Plugins/Runtime/MultiServerReplication/Source/MultiServerReplication/Private/MultiServerBeaconClient.cpp b/Engine/Plugins/Runtime/MultiServerReplication/Source/MultiServerReplication/Private/MultiServerBeaconClient.cpp
--- a/Engine/Plugins/Runtime/MultiServerReplication/Source/MultiServerReplication/Private/MultiServerBeaconClient.cpp
+++ b/Engine/Plugins/Runtime/MultiServerReplication/Source/MultiServerReplication/Private/MultiServerBeaconClient.cpp
@@ -112,7 +112,7 @@ void AMultiServerBeaconClient::ConnectToServer(const FString& ConnectInfo)
 
 	if (!bSuccess)
 	{
-		OnFailure();
+		AOnlineBeacon::OnFailure(EBeaconFailureReason::InitFailure, TEXTVIEW("beacon_multiserverclient_init_beacon_failure"));
 	}
 }
 
diff --git a/Engine/Plugins/Runtime/MultiServerReplication/Source/MultiServerReplication/Private/MultiServerBeaconHost.cpp b/Engine/Plugins/Runtime/MultiServerReplication/Source/MultiServerReplication/Private/MultiServerBeaconHost.cpp
--- a/Engine/Plugins/Runtime/MultiServerReplication/Source/MultiServerReplication/Private/MultiServerBeaconHost.cpp
+++ b/Engine/Plugins/Runtime/MultiServerReplication/Source/MultiServerReplication/Private/MultiServerBeaconHost.cpp
@@ -53,7 +53,7 @@ bool AMultiServerBeaconHost::InitHost()
 			{
 				// error initializing the network stack...
 				UE_LOG(LogBeacon, Log, TEXT("AMultiServerBeaconHost::InitHost failed"));
-				OnFailure();
+				OnFailure(EBeaconFailureReason::InitFailure, TEXTVIEW("beacon_multiserverhost_init_host_failure"));
 			}
 		}
 	}
