Add IOnlineSubsystemUtils::IsOnlinePIEEnabledStatic()

Static helper that checks whether online PIE is enabled, safe to call
before UObjectInitialized(). Uses CDO when available, falls back to
reading INI and command line auth credentials directly.

#Synced-CL 52092337

diff --git a/Engine/Plugins/Online/OnlineSubsystemUtils/Source/OnlineSubsystemUtils/Private/OnlineSubsystemUtilsModule.cpp b/Engine/Plugins/Online/OnlineSubsystemUtils/Source/OnlineSubsystemUtils/Private/OnlineSubsystemUtilsModule.cpp
--- a/Engine/Plugins/Online/OnlineSubsystemUtils/Source/OnlineSubsystemUtils/Private/OnlineSubsystemUtilsModule.cpp
+++ b/Engine/Plugins/Online/OnlineSubsystemUtils/Source/OnlineSubsystemUtils/Private/OnlineSubsystemUtilsModule.cpp
@@ -9,13 +9,24 @@
 
 #if WITH_EDITOR
 #include "Misc/CommandLine.h"
+#include "Misc/ConfigCacheIni.h"
 #include "Misc/Parse.h"
 #endif
 
 IMPLEMENT_MODULE(FOnlineSubsystemUtilsModule, OnlineSubsystemUtils);
 
+#if WITH_EDITOR
+static bool HasAuthCommandLineStatic()
+{
+	FString Id, Token, Type;
+	return FParse::Value(FCommandLine::Get(), TEXT("AUTH_LOGIN="), Id) &&
+		   FParse::Value(FCommandLine::Get(), TEXT("AUTH_PASSWORD="), Token) &&
+		   FParse::Value(FCommandLine::Get(), TEXT("AUTH_TYPE="), Type);
+}
+#endif
+
 /**
- * Concrete implementation of IOnlineSubsystemUtils interface 
+ * Concrete implementation of IOnlineSubsystemUtils interface
  */
 class FOnlineSubsystemUtils : public IOnlineSubsystemUtils
 {
@@ -116,17 +127,12 @@ private:
 			   FParse::Value(FCommandLine::Get(), TEXT("AUTH_PASSWORD="), LoginParameters.Token) &&
 			   FParse::Value(FCommandLine::Get(), TEXT("AUTH_TYPE="), LoginParameters.Type);
 	}
-	bool HasAuthCommandLine() const
-	{
-		FOnlineAccountCredentials LoginParameters;
-		return HasAuthCommandLine(LoginParameters);
-	}
-public:	
+public:
 	virtual bool SupportsOnlinePIE() const override
 	{
 		check(UObjectInitialized());
 		const UOnlinePIESettings* OnlinePIESettings = GetDefault<UOnlinePIESettings>();
-		if ((OnlinePIESettings->bOnlinePIEEnabled && GetNumPIELogins() > 0) || HasAuthCommandLine())
+		if ((OnlinePIESettings->bOnlinePIEEnabled && GetNumPIELogins() > 0) || HasAuthCommandLineStatic())
 		{
 			// If we can't get the identity interface then things are either not configured right or disabled
 			IOnlineIdentityPtr IdentityInt = Online::GetIdentityInterface();
@@ -155,7 +161,7 @@ public:
 			return false;
 		}
 		const UOnlinePIESettings* OnlinePIESettings = GetDefault<UOnlinePIESettings>();
-		return OnlinePIESettings->bOnlinePIEEnabled || HasAuthCommandLine();
+		return OnlinePIESettings->bOnlinePIEEnabled || HasAuthCommandLineStatic();
 	}
 
 	virtual int32 GetNumPIELogins() const override
@@ -171,7 +177,7 @@ public:
 				NumValidLogins++;
 			}
 		}
-		if (NumValidLogins == 0 && HasAuthCommandLine())
+		if (NumValidLogins == 0 && HasAuthCommandLineStatic())
 		{
 			NumValidLogins++;
 		}
@@ -293,3 +299,32 @@ void FOnlineSubsystemUtilsModule::ShutdownModule()
 	SubsystemUtils.Reset();
 	OnlineServicesUtils.Reset();
 }
+
+#if WITH_EDITOR
+bool IOnlineSubsystemUtils::IsOnlinePIEEnabledStatic()
+{
+	if (UObjectInitialized())
+	{
+		if (IOnlineSubsystemUtils* OnlineUtils = Online::GetUtils())
+		{
+			return OnlineUtils->IsOnlinePIEEnabled();
+		}
+	}
+
+	// Fallback: read INI directly (CDO not available yet)
+	bool bOnlinePIEEnabled = false;
+	GConfig->GetBool(TEXT("/Script/OnlineSubsystemUtils.OnlinePIESettings"), TEXT("bOnlinePIEEnabled"), bOnlinePIEEnabled, GEditorPerProjectIni);
+
+	if (bOnlinePIEEnabled)
+	{
+		TArray<FString> PIELogins;
+		GConfig->GetArray(TEXT("/Script/OnlineSubsystemUtils.OnlinePIESettings"), TEXT("Logins"), PIELogins, GEditorPerProjectIni);
+		if (PIELogins.Num() == 0)
+		{
+			bOnlinePIEEnabled = false;
+		}
+	}
+
+	return bOnlinePIEEnabled || HasAuthCommandLineStatic();
+}
+#endif
diff --git a/Engine/Plugins/Online/OnlineSubsystemUtils/Source/OnlineSubsystemUtils/Public/OnlineSubsystemUtils.h b/Engine/Plugins/Online/OnlineSubsystemUtils/Source/OnlineSubsystemUtils/Public/OnlineSubsystemUtils.h
--- a/Engine/Plugins/Online/OnlineSubsystemUtils/Source/OnlineSubsystemUtils/Public/OnlineSubsystemUtils.h
+++ b/Engine/Plugins/Online/OnlineSubsystemUtils/Source/OnlineSubsystemUtils/Public/OnlineSubsystemUtils.h
@@ -151,6 +151,12 @@ public:
 	virtual int32 GetNumPIELogins() const = 0;
 	/** @return the array of valid credentials the user has setup for Play In Editor (PIE) */
 	virtual void GetPIELogins(TArray<FOnlineAccountCredentials>& Logins) = 0;
+
+	/**
+	 * Static check for whether online PIE is enabled. Safe to call at any point during startup —
+	 * uses CDO if UObjectInitialized(), falls back to reading INI + command line auth credentials.
+	 */
+	static ONLINESUBSYSTEMUTILS_API bool IsOnlinePIEEnabledStatic();
 #endif // WITH_EDITOR
 };
 
